Move AssetPlacementEdMode to experimental plugin from editor module

Also removed the editor setting and moved the visibility of the mode to be tied to the SM instance elements enabled cvar.

diff --git a/Engine/Plugins/Experimental/AssetPlacementEdMode/AssetPlacementEdMode.uplugin b/Engine/Plugins/Experimental/AssetPlacementEdMode/AssetPlacementEdMode.uplugin
new file mode 100644
--- /dev/null
+++ b/Engine/Plugins/Experimental/AssetPlacementEdMode/AssetPlacementEdMode.uplugin
@@ -0,0 +1,24 @@
+{
+	"FileVersion": 3,
+	"Version": 1,
+	"VersionName": "0.1",
+	"FriendlyName": "Asset Placement Editor Mode",
+	"Description": "The asset placement editor mode enables quick placement of assets into a world as full actors or various types of lighter weight objects, like batched render instances.",
+	"Category": "Editor",
+	"CreatedBy": "Epic Games, Inc.",
+	"CreatedByURL": "http://epicgames.com",
+	"DocsURL": "",
+	"MarketplaceURL": "",
+	"SupportURL": "",
+	"CanContainContent": true,
+	"IsBetaVersion": true,
+	"Installed": false,
+	"Modules": [
+		{
+			"Name": "AssetPlacementEdMode",
+			"Type": "Editor",
+			"LoadingPhase": "Default"
+		}
+	],
+	"Plugins": []
+}
\ No newline at end of file
diff --git a/Engine/Plugins/Experimental/AssetPlacementEdMode/Source/AssetPlacementEdMode/AssetPlacementEdMode.Build.cs b/Engine/Plugins/Experimental/AssetPlacementEdMode/Source/AssetPlacementEdMode/AssetPlacementEdMode.Build.cs
new file mode 100644
--- /dev/null
+++ b/Engine/Plugins/Experimental/AssetPlacementEdMode/Source/AssetPlacementEdMode/AssetPlacementEdMode.Build.cs
@@ -0,0 +1,33 @@
+// Copyright Epic Games, Inc. All Rights Reserved.
+
+using UnrealBuildTool;
+
+public class AssetPlacementEdMode : ModuleRules
+{
+	public AssetPlacementEdMode(ReadOnlyTargetRules Target) : base(Target)
+	{
+		PrivateDependencyModuleNames.AddRange(
+			new string[] {
+				"ContentBrowser",
+				"Core",
+				"CoreUObject",
+                "InputCore",
+				"Engine",
+				"Slate",
+				"SlateCore",
+                "EditorStyle",
+				"EditorFramework",
+				"EditorSubsystem",
+				"UnrealEd",
+				"EditorInteractiveToolsFramework",
+				"InteractiveToolsFramework",
+				"Foliage",
+				"Landscape",
+				"TypedElementFramework",
+				"TypedElementRuntime",
+				"PropertyEditor",
+				"EditorWidgets",
+			}
+		);
+	}
+}
diff --git a/Engine/Plugins/Experimental/AssetPlacementEdMode/Source/AssetPlacementEdMode/Private/AssetPlacementEdMode.cpp b/Engine/Plugins/Experimental/AssetPlacementEdMode/Source/AssetPlacementEdMode/Private/AssetPlacementEdMode.cpp
new file mode 100644
--- /dev/null
+++ b/Engine/Plugins/Experimental/AssetPlacementEdMode/Source/AssetPlacementEdMode/Private/AssetPlacementEdMode.cpp
@@ -0,0 +1,240 @@
+// Copyright Epic Games, Inc. All Rights Reserved.
+
+#include "AssetPlacementEdMode.h"
+#include "AssetPlacementEdModeToolkit.h"
+#include "InteractiveToolManager.h"
+#include "AssetPlacementEdModeStyle.h"
+#include "AssetPlacementEdModeCommands.h"
+#include "AssetPlacementSettings.h"
+#include "EdModeInteractiveToolsContext.h"
+#include "EditorModeManager.h"
+#include "Selection.h"
+#include "EngineUtils.h"
+
+#include "Tools/AssetEditorContextInterface.h"
+#include "Tools/PlacementSelectTool.h"
+#include "Tools/PlacementLassoSelectTool.h"
+#include "Tools/PlacementPlaceTool.h"
+#include "Tools/PlacementPlaceSingleTool.h"
+#include "Tools/PlacementEraseTool.h"
+
+#include "Factories/AssetFactoryInterface.h"
+#include "Elements/Framework/EngineElementsLibrary.h"
+#include "Elements/Framework/TypedElementCommonActions.h"
+#include "Elements/Framework/TypedElementRegistry.h"
+#include "Elements/SMInstance/SMInstanceElementData.h" // For SMInstanceElementDataUtil::SMInstanceElementsEnabled
+
+#include "Settings/LevelEditorMiscSettings.h"
+#include "Modes/PlacementModeSubsystem.h"
+
+#define LOCTEXT_NAMESPACE "AssetPlacementEdMode"
+
+constexpr TCHAR UAssetPlacementEdMode::AssetPlacementEdModeID[];
+
+UAssetPlacementEdMode::UAssetPlacementEdMode()
+{
+	TAttribute<bool> IsEnabledAttr = TAttribute<bool>::Create(TAttribute<bool>::FGetter::CreateStatic(&UAssetPlacementEdMode::IsEnabled));
+	Info = FEditorModeInfo(UAssetPlacementEdMode::AssetPlacementEdModeID,
+		LOCTEXT("AssetPlacementEdModeName", "Placement"),
+		FSlateIcon(FAssetPlacementEdModeStyle::Get().GetStyleSetName(), "LevelEditor.AssetPlacementEdMode"),
+		MoveTemp(IsEnabledAttr));
+}
+
+UAssetPlacementEdMode::~UAssetPlacementEdMode()
+{
+}
+
+void UAssetPlacementEdMode::Enter()
+{
+	Super::Enter();
+	if (UPlacementModeSubsystem* PlacementModeSubsystem = GEditor->GetEditorSubsystem<UPlacementModeSubsystem>())
+	{
+		SettingsObjectAsPlacementSettings = PlacementModeSubsystem->GetMutableModeSettingsObject();
+		SettingsObjectAsPlacementSettings->LoadSettings();
+		Toolkit->SetModeSettingsObject(SettingsObjectAsPlacementSettings.Get());
+	}
+
+	const FAssetPlacementEdModeCommands& PlacementModeCommands = FAssetPlacementEdModeCommands::Get();
+	RegisterTool(PlacementModeCommands.Select, UPlacementModeSelectTool::ToolName, NewObject<UPlacementModeSelectToolBuilder>(this));
+	RegisterTool(PlacementModeCommands.Place, UPlacementModePlacementTool::ToolName, NewObject<UPlacementModePlacementToolBuilder>(this));
+	RegisterTool(PlacementModeCommands.LassoSelect, UPlacementModeLassoSelectTool::ToolName, NewObject<UPlacementModeLassoSelectToolBuilder>(this));
+	RegisterTool(PlacementModeCommands.PlaceSingle, UPlacementModePlaceSingleTool::ToolName, NewObject<UPlacementModePlaceSingleToolBuilder>(this));
+	RegisterTool(PlacementModeCommands.Erase, UPlacementModeEraseTool::ToolName, NewObject<UPlacementModeEraseToolBuilder>(this));
+
+	// Stash the current editor selection, since this mode will modify it.
+	Owner->StoreSelection(AssetPlacementEdModeID);
+	bIsInSelectionTool = false;
+
+	// Disable undo tracking so that we can't accidentally undo ourselves out of the select mode and into an invalid state.
+	GetToolManager()->ConfigureChangeTrackingMode(EToolChangeTrackingMode::NoChangeTracking);
+
+	// Enable the select tool by default.
+	GetInteractiveToolsContext()->StartTool(UPlacementModeSelectTool::ToolName);
+
+	SMInstanceElementDataUtil::OnSMInstanceElementsEnabledChanged().AddUObject(this, &UAssetPlacementEdMode::OnSMIsntancedElementsEnabledChanged);
+}
+
+void UAssetPlacementEdMode::Exit()
+{
+	SMInstanceElementDataUtil::OnSMInstanceElementsEnabledChanged().RemoveAll(this);
+
+	Super::Exit();
+
+	SettingsObjectAsPlacementSettings->SaveSettings();
+	SettingsObjectAsPlacementSettings.Reset();
+
+	// Restore the selection to the original state after all the tools have shutdown in UEdMode::Exit()
+	// Since they can continue messing with selection states.
+	Owner->RestoreSelection(AssetPlacementEdModeID);
+}
+
+void UAssetPlacementEdMode::CreateToolkit()
+{
+	Toolkit = MakeShared<FAssetPlacementEdModeToolkit>();
+}
+
+bool UAssetPlacementEdMode::UsesToolkits() const
+{
+	return true;
+}
+
+TMap<FName, TArray<TSharedPtr<FUICommandInfo>>> UAssetPlacementEdMode::GetModeCommands() const
+{
+	return FAssetPlacementEdModeCommands::Get().GetCommands();
+}
+
+void UAssetPlacementEdMode::BindCommands()
+{
+	const TSharedRef<FUICommandList>& CommandList = Toolkit->GetToolkitCommands();
+	const FAssetPlacementEdModeCommands& PlacementModeCommands = FAssetPlacementEdModeCommands::Get();
+}
+
+bool UAssetPlacementEdMode::IsSelectionAllowed(AActor* InActor, bool bInSelection) const
+{
+	// Always allow deselection, for stashing selection set.
+	if (!bInSelection)
+	{
+		return true;
+	}
+
+	// Otherwise, need to be in selection tool for selection to be allowed
+	if (!bIsInSelectionTool)
+	{
+		return false;
+	}
+
+	// And we need to have a valid palette item.
+	FTypedElementHandle ActorHandle = UEngineElementsLibrary::AcquireEditorActorElementHandle(InActor);
+	return SettingsObjectAsPlacementSettings->DoesActivePaletteSupportElement(ActorHandle);
+}
+
+void UAssetPlacementEdMode::OnToolStarted(UInteractiveToolManager* Manager, UInteractiveTool* Tool)
+{
+	Super::OnToolStarted(Manager, Tool);
+
+	bool bWasInSelectionTool = bIsInSelectionTool;
+	FString ActiveToolName = GetToolManager()->GetActiveToolName(EToolSide::Mouse);
+	bool bIsSinglePlaceTool = ActiveToolName == UPlacementModePlaceSingleTool::ToolName;
+	if (ActiveToolName == UPlacementModeSelectTool::ToolName || ActiveToolName == UPlacementModeLassoSelectTool::ToolName || bIsSinglePlaceTool)
+	{
+		bIsInSelectionTool = true;
+	}
+	else
+	{
+		bIsInSelectionTool = false;
+	}
+
+	// Restore the selection if we're going into the selection tools.
+	// Allow the selection to be empty if we're going into single place tool for a clean slate.
+	bool bRestoreSelectionState = bIsInSelectionTool && !bWasInSelectionTool && !bIsSinglePlaceTool;
+	if (bRestoreSelectionState)
+	{
+		Owner->RestoreSelection(UPlacementModeSelectTool::ToolName);
+	}
+	else if (!bIsInSelectionTool)
+	{
+		// If we can't select, clear out the selection set for the active tool.
+		ClearSelection();
+	}
+}
+
+void UAssetPlacementEdMode::OnToolEnded(UInteractiveToolManager* Manager, UInteractiveTool* Tool)
+{
+	Super::OnToolEnded(Manager, Tool);
+
+	// Always store the most recent selection, even if we are leaving single placement tool to preserve what the user was doing last.
+	if (bIsInSelectionTool)
+	{
+		constexpr bool bClearSelection = false;
+		Owner->StoreSelection(UPlacementModeSelectTool::ToolName, false);
+	}
+}
+
+bool UAssetPlacementEdMode::UsesPropertyWidgets() const
+{
+	return IsInSelectionTool();
+}
+
+bool UAssetPlacementEdMode::ShouldDrawWidget() const
+{
+	if (!IsInSelectionTool())
+	{
+		return false;
+	}
+
+	// Disable the widget in the lasso tool. The drag operations currently fight with the widget.
+	if (GetToolManager()->GetActiveToolName(EToolSide::Mouse) == UPlacementModeLassoSelectTool::ToolName)
+	{
+		return false;
+	}
+
+	return Super::ShouldDrawWidget();
+}
+
+bool UAssetPlacementEdMode::IsEnabled()
+{
+	return SMInstanceElementDataUtil::SMInstanceElementsEnabled();
+}
+
+void UAssetPlacementEdMode::ClearSelection()
+{
+	GetToolManager()->BeginUndoTransaction(LOCTEXT("PlacementClearSelection", "Clear Selection"));
+	if (UTypedElementSelectionSet* SelectionSet = Owner->GetEditorSelectionSet())
+	{
+		SelectionSet->ClearSelection(FTypedElementSelectionOptions());
+	}
+	GetToolManager()->EndUndoTransaction();
+}
+
+bool UAssetPlacementEdMode::HasAnyAssetsInPalette() const
+{
+	return SettingsObjectAsPlacementSettings.IsValid() ? (SettingsObjectAsPlacementSettings->GetActivePaletteItems().Num() > 0) : false;
+}
+
+bool UAssetPlacementEdMode::HasActiveSelection() const
+{
+	if (!HasAnyAssetsInPalette())
+	{
+		return false;
+	}
+
+	if (Owner->GetEditorSelectionSet()->HasSelectedElements())
+	{
+		return true;
+	}
+
+	return false;
+}
+
+bool UAssetPlacementEdMode::IsInSelectionTool() const
+{
+	return bIsInSelectionTool;
+}
+
+void UAssetPlacementEdMode::OnSMIsntancedElementsEnabledChanged()
+{
+	// Disable this mode if the SM instance element cvar changed
+	Owner->DeactivateMode(UAssetPlacementEdMode::AssetPlacementEdModeID);
+}
+
+#undef LOCTEXT_NAMESPACE
diff --git a/Engine/Plugins/Experimental/AssetPlacementEdMode/Source/AssetPlacementEdMode/Private/AssetPlacementEdModeCommands.cpp b/Engine/Plugins/Experimental/AssetPlacementEdMode/Source/AssetPlacementEdMode/Private/AssetPlacementEdModeCommands.cpp
new file mode 100644
--- /dev/null
+++ b/Engine/Plugins/Experimental/AssetPlacementEdMode/Source/AssetPlacementEdMode/Private/AssetPlacementEdModeCommands.cpp
@@ -0,0 +1,36 @@
+// Copyright Epic Games, Inc. All Rights Reserved.
+
+#include "AssetPlacementEdModeCommands.h"
+#include "AssetPlacementEdModeStyle.h"
+
+#define LOCTEXT_NAMESPACE "AssetPlacementEdMode"
+
+FAssetPlacementEdModeCommands::FAssetPlacementEdModeCommands()
+	: TCommands<FAssetPlacementEdModeCommands>("AssetPlacementEdMode",
+		LOCTEXT("AssetPlacementEdModeCommands", "AssetPlacement Editor Mode Commands"),
+		NAME_None,
+		FAssetPlacementEdModeStyle::Get().GetStyleSetName())
+{
+}
+
+void FAssetPlacementEdModeCommands::RegisterCommands()
+{
+	TArray <TSharedPtr<FUICommandInfo>>& ToolCommands = Commands.FindOrAdd(NAME_Default);
+	UI_COMMAND(Select, "Select", "Select by clicking single assets matching the active palette.", EUserInterfaceActionType::ToggleButton, FInputChord());
+	ToolCommands.Add(Select);
+	UI_COMMAND(LassoSelect, "Lasso", "Selects asset by painting the area to select.", EUserInterfaceActionType::ToggleButton, FInputChord());
+	ToolCommands.Add(LassoSelect);
+	UI_COMMAND(Place, "Paint", "Paint mutliple assets from the active palette.", EUserInterfaceActionType::ToggleButton, FInputChord());
+	ToolCommands.Add(Place);
+	UI_COMMAND(PlaceSingle, "Single", "Place a single, random asset from the active palette.", EUserInterfaceActionType::ToggleButton, FInputChord());
+	ToolCommands.Add(PlaceSingle);
+	UI_COMMAND(Erase, "Erase", "Paint to erase assets matching the active palette.", EUserInterfaceActionType::ToggleButton, FInputChord());
+	ToolCommands.Add(Erase);
+}
+
+TMap<FName, TArray<TSharedPtr<FUICommandInfo>>> FAssetPlacementEdModeCommands::GetCommands()
+{
+	return FAssetPlacementEdModeCommands::Get().Commands;
+}
+
+#undef LOCTEXT_NAMESPACE
diff --git a/Engine/Plugins/Experimental/AssetPlacementEdMode/Source/AssetPlacementEdMode/Private/AssetPlacementEdModeModule.cpp b/Engine/Plugins/Experimental/AssetPlacementEdMode/Source/AssetPlacementEdMode/Private/AssetPlacementEdModeModule.cpp
new file mode 100644
--- /dev/null
+++ b/Engine/Plugins/Experimental/AssetPlacementEdMode/Source/AssetPlacementEdMode/Private/AssetPlacementEdModeModule.cpp
@@ -0,0 +1,36 @@
+// Copyright Epic Games, Inc. All Rights Reserved.
+
+#include "AssetPlacementEdModeModule.h"
+#include "AssetPlacementEdModeCommands.h"
+#include "AssetPlacementEdModeStyle.h"
+#include "Modules/ModuleManager.h"
+#include "AssetTypeActions_PlacementPalette.h"
+#include "AssetToolsModule.h"
+
+#define LOCTEXT_NAMESPACE "AssetPlacementEdMode"
+
+IMPLEMENT_MODULE( FAssetPlacementEdMode, AssetPlacementEdMode );
+
+void FAssetPlacementEdMode::StartupModule()
+{
+	FAssetPlacementEdModeStyle::Get();
+	FAssetPlacementEdModeCommands::Register();
+
+	FAssetToolsModule& AssetToolsModule = FModuleManager::LoadModuleChecked<FAssetToolsModule>("AssetTools");
+	PaletteAssetActions = MakeShared<FAssetTypeActions_PlacementPalette>(EAssetTypeCategories::Type::Misc);
+	AssetToolsModule.Get().RegisterAssetTypeActions(PaletteAssetActions.ToSharedRef());
+}
+
+void FAssetPlacementEdMode::ShutdownModule()
+{
+	if (FAssetToolsModule* AssetToolsModule = FModuleManager::GetModulePtr<FAssetToolsModule>("AssetTools"))
+	{
+		AssetToolsModule->Get().UnregisterAssetTypeActions(PaletteAssetActions.ToSharedRef());
+	}
+	PaletteAssetActions.Reset();
+
+	FAssetPlacementEdModeCommands::Unregister();
+	FAssetPlacementEdModeStyle::Shutdown();
+}
+
+#undef LOCTEXT_NAMESPACE
diff --git a/Engine/Plugins/Experimental/AssetPlacementEdMode/Source/AssetPlacementEdMode/Private/AssetPlacementEdModeStyle.cpp b/Engine/Plugins/Experimental/AssetPlacementEdMode/Source/AssetPlacementEdMode/Private/AssetPlacementEdModeStyle.cpp
new file mode 100644
--- /dev/null
+++ b/Engine/Plugins/Experimental/AssetPlacementEdMode/Source/AssetPlacementEdMode/Private/AssetPlacementEdModeStyle.cpp
@@ -0,0 +1,86 @@
+// Copyright Epic Games, Inc. All Rights Reserved.
+
+#include "AssetPlacementEdModeStyle.h"
+
+#include "Styling/CoreStyle.h"
+#include "Styling/SlateStyle.h"
+#include "Styling/SlateStyleMacros.h"
+#include "Styling/SlateStyleRegistry.h"
+#include "Styling/SlateTypes.h"
+#include "Styling/SlateStyleMacros.h"
+#include "Styling/StarshipCoreStyle.h"
+#include "Styling/SlateStyleRegistry.h"
+
+FName FAssetPlacementEdModeStyle::StyleName("AssetPlacementEdModeStyle");
+TUniquePtr<FAssetPlacementEdModeStyle> FAssetPlacementEdModeStyle::Inst(nullptr);
+
+const FName& FAssetPlacementEdModeStyle::GetStyleSetName() const
+{
+	return StyleName;
+}
+
+const FAssetPlacementEdModeStyle& FAssetPlacementEdModeStyle::Get()
+{
+	if (!Inst.IsValid())
+	{
+		Inst = TUniquePtr<FAssetPlacementEdModeStyle>(new FAssetPlacementEdModeStyle);
+	}
+	return *(Inst.Get());
+}
+
+void FAssetPlacementEdModeStyle::Shutdown()
+{
+	Inst.Reset();
+}
+
+FAssetPlacementEdModeStyle::FAssetPlacementEdModeStyle()
+	: FSlateStyleSet(StyleName)
+{
+
+	const FVector2D Icon20x20(20.0f, 20.0f);
+
+	SetParentStyleName("EditorStyle");
+	SetContentRoot(FPaths::EngineContentDir() / TEXT("Editor/Slate"));
+	SetCoreContentRoot(FPaths::EngineContentDir() / TEXT("Slate"));
+
+	Set("LevelEditor.AssetPlacementEdMode", new IMAGE_BRUSH_SVG("Starship/MainToolbar/AssetPlacementMode", Icon20x20));
+	Set("ClassIcon.PlacementPaletteAsset", new IMAGE_BRUSH_SVG("Starship/AssetIcons/AssetPalette", FVector2D(16.0f, 16.0f)));
+	Set("ClassThumbnail.PlacementPaletteAsset", new IMAGE_BRUSH_SVG("Starship/AssetIcons/AssetPalette", FVector2D(64.0f, 64.0f)));
+
+	Set("AssetPlacementEdMode.Select",                      new IMAGE_BRUSH("Icons/GeneralTools/Select_40x", Icon20x20));
+	Set("AssetPlacementEdMode.Select.Small",                new IMAGE_BRUSH("Icons/GeneralTools/Select_40x", Icon20x20));
+	Set("AssetPlacementEdMode.SelectInvalid",               new IMAGE_BRUSH("Icons/GeneralTools/SelectInvalid_40x", Icon20x20));
+	Set("AssetPlacementEdMode.SelectInvalid.Small",         new IMAGE_BRUSH("Icons/GeneralTools/SelectInvalid_40x", Icon20x20));
+	Set("AssetPlacementEdMode.LassoSelect",                 new IMAGE_BRUSH("Icons/GeneralTools/Lasso_40x", Icon20x20));
+	Set("AssetPlacementEdMode.LassoSelect.Small",           new IMAGE_BRUSH("Icons/GeneralTools/Lasso_40x", Icon20x20));
+	Set("AssetPlacementEdMode.PlaceSingle",                 new IMAGE_BRUSH("Icons/GeneralTools/Foliage_40x", Icon20x20));
+	Set("AssetPlacementEdMode.PlaceSingle.Small",           new IMAGE_BRUSH("Icons/GeneralTools/Foliage_40x", Icon20x20));
+	Set("AssetPlacementEdMode.Place",                       new IMAGE_BRUSH("Icons/GeneralTools/Paint_40x", Icon20x20));
+	Set("AssetPlacementEdMode.Place.Small",                 new IMAGE_BRUSH("Icons/GeneralTools/Paint_40x", Icon20x20));
+	Set("AssetPlacementEdMode.ReapplySettings",             new IMAGE_BRUSH("Icons/GeneralTools/Repaint_40x", Icon20x20));
+	Set("AssetPlacementEdMode.ReapplySettings.Small",       new IMAGE_BRUSH("Icons/GeneralTools/Repaint_40x", Icon20x20));
+	Set("AssetPlacementEdMode.Fill",                        new IMAGE_BRUSH("Icons/GeneralTools/PaintBucket_40x", Icon20x20));
+	Set("AssetPlacementEdMode.Fill.Small",                  new IMAGE_BRUSH("Icons/GeneralTools/PaintBucket_40x", Icon20x20));
+	Set("AssetPlacementEdMode.Erase",                       new IMAGE_BRUSH("Icons/GeneralTools/Erase_40x", Icon20x20));
+	Set("AssetPlacementEdMode.Erase.Small",                 new IMAGE_BRUSH("Icons/GeneralTools/Erase_40x", Icon20x20));
+	Set("AssetPlacementEdMode.Filter",                      new IMAGE_BRUSH("Icons/GeneralTools/Filter_40x", Icon20x20));
+	Set("AssetPlacementEdMode.Filter.Small",                new IMAGE_BRUSH("Icons/GeneralTools/Filter_40x", Icon20x20));
+	Set("AssetPlacementEdMode.Settings",                    new IMAGE_BRUSH("Icons/GeneralTools/Settings_40x", Icon20x20));
+	Set("AssetPlacementEdMode.Settings.Small",              new IMAGE_BRUSH("Icons/GeneralTools/Settings_40x", Icon20x20));
+	Set("AssetPlacementEdMode.MoveToActivePartition",       new IMAGE_BRUSH("Icons/GeneralTools/MoveToLevel_40x", Icon20x20));
+	Set("AssetPlacementEdMode.MoveToActivePartition.Small", new IMAGE_BRUSH("Icons/GeneralTools/MoveToLevel_40x", Icon20x20));
+
+	Set("AssetPlacementEdMode.AddAssetType.Text", FTextBlockStyle(GetWidgetStyle<FTextBlockStyle>("NormalText"))
+		.SetFont(DEFAULT_FONT("Bold", 10))
+		.SetColorAndOpacity(FLinearColor(1.0f, 1.0f, 1.0f))
+		.SetHighlightColor(FLinearColor(1.0f, 1.0f, 1.0f))
+		.SetShadowOffset(FVector2D(1, 1))
+		.SetShadowColorAndOpacity(FLinearColor(0, 0, 0, 0.9f)));
+
+	FSlateStyleRegistry::RegisterSlateStyle(*this);
+}
+
+FAssetPlacementEdModeStyle::~FAssetPlacementEdModeStyle()
+{
+	FSlateStyleRegistry::UnRegisterSlateStyle(*this);
+}
diff --git a/Engine/Plugins/Experimental/AssetPlacementEdMode/Source/AssetPlacementEdMode/Private/AssetPlacementEdModeStyle.h b/Engine/Plugins/Experimental/AssetPlacementEdMode/Source/AssetPlacementEdMode/Private/AssetPlacementEdModeStyle.h
new file mode 100644
--- /dev/null
+++ b/Engine/Plugins/Experimental/AssetPlacementEdMode/Source/AssetPlacementEdMode/Private/AssetPlacementEdModeStyle.h
@@ -0,0 +1,25 @@
+// Copyright Epic Games, Inc. All Rights Reserved.
+
+#pragma once
+
+#include "CoreMinimal.h"
+#include "Styling/SlateStyle.h"
+#include "Styling/SlateStyleRegistry.h"
+
+class FAssetPlacementEdModeStyle
+	: public FSlateStyleSet
+{
+public:
+	virtual const FName& GetStyleSetName() const override;
+
+	static const FAssetPlacementEdModeStyle& Get();
+	static void Shutdown();
+	
+	~FAssetPlacementEdModeStyle();
+
+private:
+	FAssetPlacementEdModeStyle();
+
+	static FName StyleName;
+	static TUniquePtr<FAssetPlacementEdModeStyle> Inst;
+};
diff --git a/Engine/Plugins/Experimental/AssetPlacementEdMode/Source/AssetPlacementEdMode/Private/AssetPlacementEdModeToolkit.cpp b/Engine/Plugins/Experimental/AssetPlacementEdMode/Source/AssetPlacementEdMode/Private/AssetPlacementEdModeToolkit.cpp
new file mode 100644
--- /dev/null
+++ b/Engine/Plugins/Experimental/AssetPlacementEdMode/Source/AssetPlacementEdMode/Private/AssetPlacementEdModeToolkit.cpp
@@ -0,0 +1,129 @@
+// Copyright Epic Games, Inc. All Rights Reserved.
+
+#include "AssetPlacementEdModeToolkit.h"
+
+#include "IDetailsView.h"
+#include "SAssetPlacementPalette.h"
+#include "Widgets/Docking/SDockTab.h"
+#include "Widgets/Layout/SExpandableArea.h"
+#include "Widgets/Layout/SScrollBox.h"
+#include "Toolkits/AssetEditorModeUILayer.h"
+#include "PropertyEditorModule.h"
+#include "Modules/ModuleManager.h"
+
+#define LOCTEXT_NAMESPACE "AssetPlacementEdModeToolkit"
+
+FAssetPlacementEdModeToolkit::FAssetPlacementEdModeToolkit()
+{
+	// Create the details view for the details view tab, so that it can be shared across our custom tabs
+	{
+		FPropertyEditorModule& PropertyEditorModule = FModuleManager::GetModuleChecked<FPropertyEditorModule>("PropertyEditor");
+		FDetailsViewArgs DetailsViewArgs;
+		DetailsViewArgs.bAllowSearch = false;
+		DetailsViewArgs.NameAreaSettings = FDetailsViewArgs::ENameAreaSettings::HideNameArea;
+		DetailsViewArgs.bHideSelectionTip = true;
+
+		PaletteItemDetailsView = PropertyEditorModule.CreateDetailView(DetailsViewArgs);
+	}
+}
+
+void FAssetPlacementEdModeToolkit::GetToolPaletteNames(TArray<FName>& PaletteNames) const
+{
+	PaletteNames.Add(NAME_Default);
+}
+
+
+FName FAssetPlacementEdModeToolkit::GetToolkitFName() const
+{
+	return FName("AssetPlacementEdMode");
+}
+
+FText FAssetPlacementEdModeToolkit::GetBaseToolkitName() const
+{
+	return LOCTEXT("DisplayName", "AssetPlacementEdMode Tool");
+}
+
+void FAssetPlacementEdModeToolkit::InvokeUI()
+{
+	FModeToolkit::InvokeUI();
+
+	if (ModeUILayer.IsValid())
+	{
+		TSharedPtr<FAssetEditorModeUILayer> ModeUILayerPtr = ModeUILayer.Pin();
+		AssetPaletteTab = ModeUILayerPtr->GetTabManager()->TryInvokeTab(UAssetEditorUISubsystem::BottomLeftTabID);
+
+		PaletteItemDetailsViewTab = ModeUILayerPtr->GetTabManager()->TryInvokeTab(UAssetEditorUISubsystem::BottomRightTabID);
+	}
+}
+
+void FAssetPlacementEdModeToolkit::RequestModeUITabs()
+{
+	FModeToolkit::RequestModeUITabs();
+	if (ModeUILayer.IsValid())
+	{
+		TSharedPtr<FAssetEditorModeUILayer> ModeUILayerPtr = ModeUILayer.Pin();
+		AssetPaletteInfo.OnSpawnTab = FOnSpawnTab::CreateSP(SharedThis(this), &FAssetPlacementEdModeToolkit::CreateAssetPalette);
+		AssetPaletteInfo.TabLabel = LOCTEXT("AssetPaletteTab", "Palette");
+		AssetPaletteInfo.TabTooltip = LOCTEXT("AssetPaletteTabTooltipText", "Open the asset palette tab, which contains placement mode's palette settings.");
+		AssetPaletteInfo.TabIcon = GetEditorModeIcon();
+		ModeUILayerPtr->SetModePanelInfo(UAssetEditorUISubsystem::BottomLeftTabID, AssetPaletteInfo);
+
+		PaletteItemDetailsViewInfo.OnSpawnTab = FOnSpawnTab::CreateSP(SharedThis(this), &FAssetPlacementEdModeToolkit::CreatePaletteItemDetailsView);
+		PaletteItemDetailsViewInfo.TabLabel = LOCTEXT("PaletteItemDetailsTab", "Palette Details");
+		PaletteItemDetailsViewInfo.TabTooltip = LOCTEXT("PaletteItemDetailsTabTooltipText", "Open the asset palette details tab, which allows customization of individual items in the active palette.");
+		PaletteItemDetailsViewInfo.TabIcon = GetEditorModeIcon();
+		ModeUILayerPtr->SetModePanelInfo(UAssetEditorUISubsystem::BottomRightTabID, PaletteItemDetailsViewInfo);
+	}
+}
+
+
+
+TSharedRef<SDockTab> FAssetPlacementEdModeToolkit::CreateAssetPalette(const FSpawnTabArgs& Args)
+{
+	return SNew(SDockTab)
+		[
+			SNew(SScrollBox)
+			+ SScrollBox::Slot()
+			[
+				SNew(SAssetPlacementPalette)
+				.ItemDetailsView(PaletteItemDetailsView)
+			]
+		];
+}
+
+TSharedRef<SDockTab> FAssetPlacementEdModeToolkit::CreatePaletteItemDetailsView(const FSpawnTabArgs& Args)
+{
+	return SNew(SDockTab)
+	[
+		SNew(SScrollBox)
+		+SScrollBox::Slot()
+		[
+			SNew(SOverlay)
+			+SOverlay::Slot()
+			[
+				PaletteItemDetailsView.ToSharedRef()
+			]
+
+			+SOverlay::Slot()
+			.Padding(0.0f, 4.0f, 0.0f, 0.0f)
+			.HAlign(HAlign_Center)
+			[
+				SNew(STextBlock)
+				.Text(LOCTEXT("PaletteDetailsViewHintText", "Select an item from the palette to edit settings."))
+				.Visibility(this, &FAssetPlacementEdModeToolkit::GetDetailsViewHintTextVisibility)
+				.TextStyle(FAppStyle::Get(), "HintText")
+			]
+		]
+	];
+}
+
+EVisibility FAssetPlacementEdModeToolkit::GetDetailsViewHintTextVisibility() const
+{
+	if (PaletteItemDetailsView && (PaletteItemDetailsView->GetSelectedObjects().Num() == 0))
+	{
+		return EVisibility::HitTestInvisible;
+	}
+	return EVisibility::Collapsed;
+}
+
+#undef LOCTEXT_NAMESPACE
diff --git a/Engine/Plugins/Experimental/AssetPlacementEdMode/Source/AssetPlacementEdMode/Private/AssetPlacementPaletteItem.cpp b/Engine/Plugins/Experimental/AssetPlacementEdMode/Source/AssetPlacementEdMode/Private/AssetPlacementPaletteItem.cpp
new file mode 100644
--- /dev/null
+++ b/Engine/Plugins/Experimental/AssetPlacementEdMode/Source/AssetPlacementEdMode/Private/AssetPlacementPaletteItem.cpp
@@ -0,0 +1,287 @@
+// Copyright Epic Games, Inc. All Rights Reserved.
+
+#include "AssetPlacementPaletteItem.h"
+#include "Widgets/Input/SButton.h"
+#include "Widgets/SToolTip.h"
+#include "Widgets/Input/SCheckBox.h"
+#include "Modules/ModuleManager.h"
+#include "Widgets/Images/SImage.h"
+#include "EditorStyleSet.h"
+
+#include "AssetThumbnail.h"
+#include "IAssetTools.h"
+#include "IAssetTypeActions.h"
+#include "AssetToolsModule.h"
+#include "FoliageType.h"
+#include "SAssetPlacementPalette.h"
+#include "Engine/Blueprint.h"
+#include "Factories/AssetFactoryInterface.h"
+#include "AssetPlacementSettings.h"
+#include "Instances/InstancedPlacementClientInfo.h"
+
+#define LOCTEXT_NAMESPACE "AssetPlacementMode"
+
+FPaletteItemUIInfo::FPaletteItemUIInfo(const FAssetData& InAssetData, const UPlacementPaletteClient* InPaletteItem)
+	: AssetData(InAssetData)
+	, SettingsObject(InPaletteItem ? InPaletteItem->SettingsObject.Get() : nullptr)
+{
+	check(AssetData.IsValid());
+}
+
+FAssetPlacementPaletteItemModel::FAssetPlacementPaletteItemModel(const FAssetData& InAssetData, const UPlacementPaletteClient* InPaletteItem, TSharedRef<SAssetPlacementPalette> InParentPalette, TSharedPtr<FAssetThumbnailPool> InThumbnailPool)
+	: TypeInfo(MakeShared<FPaletteItemUIInfo>(InAssetData, InPaletteItem))
+	, AssetPalette(InParentPalette)
+{
+	check(TypeInfo);
+	DisplayFName = TypeInfo->AssetData.AssetName;
+
+	int32 MaxThumbnailSize = PlacementPaletteConstants::ThumbnailSizeRange.Max;
+	TSharedPtr<FAssetThumbnail> Thumbnail = MakeShared<FAssetThumbnail>(TypeInfo->AssetData, MaxThumbnailSize, MaxThumbnailSize, InThumbnailPool);
+	
+	FAssetThumbnailConfig ThumbnailConfig;
+	FAssetToolsModule& AssetToolsModule = FModuleManager::LoadModuleChecked<FAssetToolsModule>(TEXT("AssetTools"));
+	if (!IsBlueprint())
+	{
+		AssetTypeActions = AssetToolsModule.Get().GetAssetTypeActionsForClass(TypeInfo->AssetData.GetClass());
+	}
+	else
+	{
+		AssetTypeActions = AssetToolsModule.Get().GetAssetTypeActionsForClass(UBlueprint::StaticClass());
+	}
+
+	if (AssetTypeActions.IsValid())
+	{
+		ThumbnailConfig.AssetTypeColorOverride = AssetTypeActions.Pin()->GetTypeColor();
+	}
+
+	ThumbnailWidget = Thumbnail->MakeThumbnailWidget(ThumbnailConfig);
+}
+
+FAssetPlacementUIInfoPtr FAssetPlacementPaletteItemModel::GetTypeUIInfo() const
+{
+	return TypeInfo;
+}
+
+TSharedRef<SWidget> FAssetPlacementPaletteItemModel::GetThumbnailWidget() const
+{
+	return ThumbnailWidget.ToSharedRef();
+}
+
+TSharedRef<SToolTip> FAssetPlacementPaletteItemModel::CreateTooltipWidget() const
+{
+	return 
+		SNew(SToolTip)
+		.TextMargin(1)
+		.BorderImage(FEditorStyle::GetBrush("ContentBrowser.TileViewTooltip.ToolTipBorder"))
+		.Visibility(this, &FAssetPlacementPaletteItemModel::GetTooltipVisibility)
+		[
+			SNew(SBorder)
+			.Padding(3.f)
+			.BorderImage(FEditorStyle::GetBrush("ContentBrowser.TileViewTooltip.NonContentBorder"))
+			[
+				SNew(SVerticalBox)
+
+				+ SVerticalBox::Slot()
+				.AutoHeight()
+				[
+					SNew(SBorder)
+					.Padding(FMargin(6.f))
+					.HAlign(HAlign_Left)
+					.BorderImage(FEditorStyle::GetBrush("ContentBrowser.TileViewTooltip.ContentBorder"))
+					[
+						SNew(STextBlock)
+						.Text(FText::FromName(DisplayFName))
+						.Font(FEditorStyle::GetFontStyle("ContentBrowser.TileViewTooltip.NameFont"))
+						.HighlightText(this, &FAssetPlacementPaletteItemModel::GetPaletteSearchText)
+					]
+				]
+				
+				+ SVerticalBox::Slot()
+				.AutoHeight()
+				.Padding(FMargin(0.f, 3.f, 0.f, 0.f))
+				[
+					
+					SNew(SHorizontalBox)
+
+					+ SHorizontalBox::Slot()
+					.AutoWidth()
+					.Padding(FMargin(0.f, 0.f, 3.f, 0.f))
+					[
+						SNew(SBorder)
+						.Padding(6.f)
+						.HAlign(HAlign_Center)
+						.Visibility(this, &FAssetPlacementPaletteItemModel::GetTooltipThumbnailVisibility)
+						.BorderImage(FEditorStyle::GetBrush("ContentBrowser.TileViewTooltip.ContentBorder"))
+						[
+							SNew(SBox)
+							.HeightOverride(64.f)
+							.WidthOverride(64.f)
+							[
+								GetThumbnailWidget()
+							]
+						]
+					]
+
+					+ SHorizontalBox::Slot()
+					[
+						SNew(SBorder)
+						.Padding(6.f)
+						.BorderImage(FEditorStyle::GetBrush("ContentBrowser.TileViewTooltip.ContentBorder"))
+						[
+							SNew(SVerticalBox)
+
+							+ SVerticalBox::Slot()
+							.Padding(0, 1)
+							.AutoHeight()
+							[
+								SNew(SHorizontalBox)
+								
+								+ SHorizontalBox::Slot()
+								.AutoWidth()
+								[
+									SNew(STextBlock)
+									.Text(LOCTEXT("SourceAssetTypeHeading", "Source Asset Type: "))
+									.ColorAndOpacity(FSlateColor::UseSubduedForeground())
+								]
+
+								+ SHorizontalBox::Slot()
+								.AutoWidth()
+								[
+									SNew(STextBlock)
+									.Text(this, &FAssetPlacementPaletteItemModel::GetSourceAssetTypeText)
+								]
+							]
+						]
+					]					
+				]
+			]
+		];
+}
+
+FName FAssetPlacementPaletteItemModel::GetDisplayFName() const
+{
+	return DisplayFName;
+}
+
+FText FAssetPlacementPaletteItemModel::GetPaletteSearchText() const
+{
+	if (AssetPalette.IsValid())
+	{
+		return AssetPalette.Pin()->GetSearchText();
+	}
+	else
+	{
+		return FText();
+	}
+}
+
+bool FAssetPlacementPaletteItemModel::IsBlueprint() const
+{
+	return TypeInfo && (TypeInfo->AssetData.GetClass()->GetDefaultObject()->IsA<UBlueprint>() || (TypeInfo->AssetData.GetClass()->ClassGeneratedBy != nullptr));
+}
+
+bool FAssetPlacementPaletteItemModel::IsAsset() const
+{
+	return TypeInfo && TypeInfo->AssetData.IsValid();
+}
+
+EVisibility FAssetPlacementPaletteItemModel::GetTooltipVisibility() const
+{
+	return (AssetPalette.IsValid() && AssetPalette.Pin()->ShouldShowTooltips()) ? EVisibility::SelfHitTestInvisible : EVisibility::Collapsed;
+}
+
+EVisibility FAssetPlacementPaletteItemModel::GetTooltipThumbnailVisibility() const
+{
+	return (AssetPalette.IsValid() && AssetPalette.Pin()->IsActiveViewMode(SAssetPlacementPalette::EViewMode::Tree)) ? EVisibility::SelfHitTestInvisible : EVisibility::Collapsed;
+}
+
+FText FAssetPlacementPaletteItemModel::GetSourceAssetTypeText() const
+{
+	if (AssetTypeActions.IsValid())
+	{
+		return AssetTypeActions.Pin()->GetName();
+	}
+
+	if (!TypeInfo->AssetData.AssetClass.IsNone())
+	{
+		return FText::FromName(TypeInfo->AssetData.AssetClass);
+	}
+
+	return FText::FromName(TypeInfo->AssetData.AssetName);
+}
+
+////////////////////////////////////////////////
+// SAssetPlacementPaletteItemTile
+////////////////////////////////////////////////
+
+const float SAssetPlacementPaletteItemTile::MinScaleForOverlayItems = 0.2f;
+
+void SAssetPlacementPaletteItemTile::Construct(const FArguments& InArgs, TSharedRef<STableViewBase> InOwnerTableView, TSharedPtr<FAssetPlacementPaletteItemModel>& InModel)
+{
+	Model = InModel;
+
+	STableRow<FAssetPlacementUIInfoPtr>::Construct(
+		STableRow<FAssetPlacementUIInfoPtr>::FArguments()
+		.Style(FEditorStyle::Get(), "ContentBrowser.AssetListView.ColumnListTableRow")
+		.Padding(1.f)
+		.Content()
+		[
+			SNew(SOverlay)
+			.ToolTip(Model->CreateTooltipWidget())
+			
+			// Thumbnail
+			+ SOverlay::Slot()
+			[
+				SNew(SBorder)
+				.Padding(4.f)
+				.BorderImage(FEditorStyle::GetBrush("ContentBrowser.ThumbnailShadow"))
+				.ForegroundColor(FLinearColor::White)
+				.ColorAndOpacity(FLinearColor::White)
+				[
+					Model->GetThumbnailWidget()
+				]
+			]
+		], InOwnerTableView);
+}
+
+////////////////////////////////////////////////
+// SAssetPlacementPaletteItemRow
+////////////////////////////////////////////////
+
+void SAssetPlacementPaletteItemRow::Construct(const FArguments& InArgs, TSharedRef<STableViewBase> InOwnerTableView, TSharedPtr<FAssetPlacementPaletteItemModel>& InModel)
+{
+	Model = InModel;
+	SMultiColumnTableRow<FAssetPlacementUIInfoPtr>::Construct(FSuperRowType::FArguments(), InOwnerTableView);
+
+	SetToolTip(Model->CreateTooltipWidget());
+}
+
+TSharedRef<SWidget> SAssetPlacementPaletteItemRow::GenerateWidgetForColumn(const FName& ColumnName)
+{
+	TSharedPtr<SWidget> TableRowContent = SNullWidget::NullWidget;
+
+	if (ColumnName == AssetPlacementPaletteTreeColumns::ColumnID_Type)
+	{
+		TableRowContent =
+			SNew(SHorizontalBox)
+			
+			+SHorizontalBox::Slot()
+			.AutoWidth()
+			[
+				SNew(SExpanderArrow, SharedThis(this))
+			]
+			
+			+SHorizontalBox::Slot()
+			.VAlign(VAlign_Center)
+			.AutoWidth()
+			[
+				SNew(STextBlock)
+				.Text(FText::FromName(Model->GetDisplayFName()))
+				.HighlightText(Model.ToSharedRef(), &FAssetPlacementPaletteItemModel::GetPaletteSearchText)
+			];
+	}
+
+	return TableRowContent.ToSharedRef();
+}
+
+#undef LOCTEXT_NAMESPACE
diff --git a/Engine/Plugins/Experimental/AssetPlacementEdMode/Source/AssetPlacementEdMode/Private/AssetPlacementPaletteItem.h b/Engine/Plugins/Experimental/AssetPlacementEdMode/Source/AssetPlacementEdMode/Private/AssetPlacementPaletteItem.h
new file mode 100644
--- /dev/null
+++ b/Engine/Plugins/Experimental/AssetPlacementEdMode/Source/AssetPlacementEdMode/Private/AssetPlacementPaletteItem.h
@@ -0,0 +1,110 @@
+// Copyright Epic Games, Inc. All Rights Reserved.
+
+#pragma once
+
+#include "CoreMinimal.h"
+#include "Misc/Attribute.h"
+#include "Layout/Visibility.h"
+#include "Input/Reply.h"
+#include "Widgets/DeclarativeSyntaxSupport.h"
+#include "Widgets/SWidget.h"
+#include "Widgets/Views/STableViewBase.h"
+#include "Widgets/Views/STableRow.h"
+#include "SAssetPlacementPalette.h"
+#include "AssetRegistry/AssetData.h"
+
+class FAssetThumbnailPool;
+class UPlacementPaletteClient;
+class IAssetTypeActions;
+class UInstancedPlacemenClientSettings;
+
+namespace PlacementPaletteConstants
+{
+	const FInt32Interval ThumbnailSizeRange(32, 128);
+}
+
+struct FPaletteItemUIInfo
+{
+	FPaletteItemUIInfo(const FAssetData& InAssetData, const UPlacementPaletteClient* InPaletteItem);
+
+	const FAssetData AssetData;
+	TWeakObjectPtr<UInstancedPlacemenClientSettings> SettingsObject;
+};
+
+typedef TSharedPtr<FPaletteItemUIInfo> FAssetPlacementUIInfoPtr;
+
+namespace AssetPlacementPaletteTreeColumns
+{
+	static const FName ColumnID_Type("Type");
+};
+
+class FAssetPlacementPaletteItemModel : public TSharedFromThis<FAssetPlacementPaletteItemModel>
+{
+public:
+	FAssetPlacementPaletteItemModel(const FAssetData& InAssetData, const UPlacementPaletteClient* InPaletteItem, TSharedRef<class SAssetPlacementPalette> InParentPalette, TSharedPtr<class FAssetThumbnailPool> InThumbnailPool);
+
+	FAssetPlacementUIInfoPtr GetTypeUIInfo() const;
+
+	/** @return The thumbnail widget for this item */
+	TSharedRef<SWidget> GetThumbnailWidget() const;
+
+	/** @return The tooltip widget for this item */
+	TSharedRef<class SToolTip> CreateTooltipWidget() const;
+
+	/** Gets the FName version of the displayed name of this item */
+	FName GetDisplayFName() const;
+
+	/** Gets the current search filter text */
+	FText GetPaletteSearchText() const;
+
+	/** @return Whether this palette item represents an instance of a foliage type blueprint class */
+	bool IsBlueprint() const;
+
+	/** @return Whether this palette item represents a foliage type asset */
+	bool IsAsset() const;
+
+private:
+	/** Gets the visibility of the entire tooltip */
+	EVisibility GetTooltipVisibility() const;
+
+	/** Gets the visibility of the thumbnail in the tooltip */
+	EVisibility GetTooltipThumbnailVisibility() const;
+
+	/** Gets the source asset type text */
+	FText GetSourceAssetTypeText() const;
+
+private:
+	TSharedPtr<SWidget> ThumbnailWidget;
+	FName DisplayFName;
+	FAssetPlacementUIInfoPtr TypeInfo;
+	TWeakPtr<SAssetPlacementPalette> AssetPalette;
+	TWeakPtr<IAssetTypeActions> AssetTypeActions;
+};
+
+/** A tile representing a foliage type in the palette */
+class SAssetPlacementPaletteItemTile : public STableRow<FAssetPlacementUIInfoPtr>
+{
+public:
+	SLATE_BEGIN_ARGS(SAssetPlacementPaletteItemTile) {}
+	SLATE_END_ARGS()
+
+	void Construct(const FArguments& InArgs, TSharedRef<STableViewBase> InOwnerTableView, TSharedPtr<FAssetPlacementPaletteItemModel>& InModel);
+
+private:
+	static const float MinScaleForOverlayItems;
+	TSharedPtr<FAssetPlacementPaletteItemModel> Model;
+};
+
+/** A tree row representing a foliage type in the palette */
+class SAssetPlacementPaletteItemRow : public SMultiColumnTableRow<FAssetPlacementUIInfoPtr>
+{
+public:
+	SLATE_BEGIN_ARGS(SAssetPlacementPaletteItemRow) {}
+	SLATE_END_ARGS()
+
+	void Construct(const FArguments& InArgs, TSharedRef<STableViewBase> InOwnerTableView, TSharedPtr<FAssetPlacementPaletteItemModel>& InModel);
+	virtual TSharedRef<SWidget> GenerateWidgetForColumn(const FName& ColumnName) override;
+
+private:
+	TSharedPtr<FAssetPlacementPaletteItemModel> Model;
+};
diff --git a/Engine/Plugins/Experimental/AssetPlacementEdMode/Source/AssetPlacementEdMode/Private/AssetPlacementSettings.cpp b/Engine/Plugins/Experimental/AssetPlacementEdMode/Source/AssetPlacementEdMode/Private/AssetPlacementSettings.cpp
new file mode 100644
--- /dev/null
+++ b/Engine/Plugins/Experimental/AssetPlacementEdMode/Source/AssetPlacementEdMode/Private/AssetPlacementSettings.cpp
@@ -0,0 +1,201 @@
+// Copyright Epic Games, Inc. All Rights Reserved.
+
+#include "AssetPlacementSettings.h"
+
+#include "Editor.h"
+#include "PackageTools.h"
+#include "PlacementPaletteAsset.h"
+#include "PlacementPaletteItem.h"
+#include "Factories/AssetFactoryInterface.h"
+#include "Subsystems/PlacementSubsystem.h"
+#include "Elements/Framework/TypedElementHandle.h"
+#include "Elements/Framework/TypedElementRegistry.h"
+#include "Elements/Interfaces/TypedElementAssetDataInterface.h"
+#include "Instances/InstancedPlacementClientInfo.h"
+
+bool UAssetPlacementSettings::CanEditChange(const FProperty* InProperty) const
+{
+	if (!Super::CanEditChange(InProperty))
+	{
+		return false;
+	}
+
+	const FName PropertyName = InProperty->GetFName();
+	if (PropertyName == GET_MEMBER_NAME_CHECKED(UAssetPlacementSettings, bAllowNegativeScale))
+	{
+		return bUseRandomScale;
+	}
+
+	if (PropertyName == GET_MEMBER_NAME_CHECKED(UAssetPlacementSettings, bAllowNegativeRotationX))
+	{
+		return bUseRandomRotationX;
+	}
+
+	if (PropertyName == GET_MEMBER_NAME_CHECKED(UAssetPlacementSettings, bAllowNegativeRotationY))
+	{
+		return bUseRandomRotationY;
+	}
+
+	if (PropertyName == GET_MEMBER_NAME_CHECKED(UAssetPlacementSettings, bAllowNegativeRotationZ))
+	{
+		return bUseRandomRotationZ;
+	}
+
+	if (PropertyName == GET_MEMBER_NAME_CHECKED(UAssetPlacementSettings, bInvertNormalAxis))
+	{
+		return bAlignToNormal;
+	}
+
+	return true;
+}
+
+void UAssetPlacementSettings::SetPaletteAsset(UPlacementPaletteAsset* InPaletteAsset)
+{
+	ActivePalette = InPaletteAsset;
+	LastActivePalettePath = FSoftObjectPath(InPaletteAsset);
+}
+
+UPlacementPaletteClient* UAssetPlacementSettings::AddClientToActivePalette(const FAssetData& InAssetData)
+{
+	UPlacementPaletteClient* NewPaletteItem = nullptr;
+	if (!InAssetData.IsValid() || !GetActivePalette())
+	{
+		return NewPaletteItem;
+	}
+
+	if (InAssetData.GetClass()->HasAnyClassFlags(CLASS_Abstract | CLASS_Deprecated | CLASS_NewerVersionExists | CLASS_NotPlaceable))
+	{
+		return NewPaletteItem;
+	}
+
+	if (!GetActivePaletteItems().FindByPredicate([InAssetData](const UPlacementPaletteClient* ItemIter) { return ItemIter ? (ItemIter->AssetPath == InAssetData.ToSoftObjectPath()) : false; }))
+	{
+		if (UPlacementSubsystem* PlacementSubystem = GEditor->GetEditorSubsystem<UPlacementSubsystem>())
+		{
+			if (TScriptInterface<IAssetFactoryInterface> AssetFactory = PlacementSubystem->FindAssetFactoryFromAssetData(InAssetData))
+			{
+				NewPaletteItem = NewObject<UPlacementPaletteClient>(GetMutableActivePalette());
+				NewPaletteItem->AssetPath = InAssetData.ToSoftObjectPath();
+				FPlacementOptions PlacementOptions;
+				PlacementOptions.InstancedPlacementGridGuid = GetActivePaletteGuid();
+				NewPaletteItem->SettingsObject = AssetFactory->FactorySettingsObjectForPlacement(InAssetData, PlacementOptions);
+				GetMutableActivePalette()->Modify();
+				if (NewPaletteItem->SettingsObject)
+				{
+					// Set the outer of the new object to the current palette
+					NewPaletteItem->SettingsObject->Rename(nullptr, GetMutableActivePalette());
+				}
+				GetMutableActivePalette()->PaletteItems.Add(NewPaletteItem);
+			}
+		}
+	}
+
+	return NewPaletteItem;
+}
+
+TArrayView<const TObjectPtr<UPlacementPaletteClient>> UAssetPlacementSettings::GetActivePaletteItems() const
+{
+	return MakeArrayView(GetActivePalette()->PaletteItems);
+}
+
+FSoftObjectPath UAssetPlacementSettings::GetActivePalettePath() const
+{
+	return FSoftObjectPath(ActivePalette);
+}
+
+const FGuid UAssetPlacementSettings::GetActivePaletteGuid() const
+{
+	return GetActivePalette()->GridGuid;
+}
+
+void UAssetPlacementSettings::ClearActivePaletteItems()
+{
+	GetMutableActivePalette()->Modify();
+	GetMutableActivePalette()->PaletteItems.Empty();
+}
+
+void UAssetPlacementSettings::LoadSettings()
+{
+	LoadConfig();
+
+	if (!UserPalette)
+	{
+		UserPalette = NewObject<UPlacementPaletteAsset>(this);
+	}
+
+	if (!UserGridGuid.IsValid())
+	{
+		FGuid::Parse(FPlatformMisc::GetLoginId(), UserGridGuid);
+	}
+
+	UserPalette->GridGuid = UserGridGuid;
+
+	ActivePalette = Cast<UPlacementPaletteAsset>(LastActivePalettePath.TryLoad());
+}
+
+void UAssetPlacementSettings::SaveSettings()
+{
+	SaveConfig();
+}
+
+void UAssetPlacementSettings::SaveActivePalette()
+{
+	if (ActivePalette)
+	{
+		UPackageTools::SavePackagesForObjects(TArray<UObject*>({ ActivePalette }));
+	}
+}
+
+bool UAssetPlacementSettings::DoesActivePaletteSupportElement(const FTypedElementHandle& InElementToCheck) const
+{	
+	if (TTypedElement<ITypedElementAssetDataInterface> AssetDataInterface = UTypedElementRegistry::GetInstance()->GetElement<ITypedElementAssetDataInterface>(InElementToCheck))
+	{
+		TArray<FAssetData> ReferencedAssetDatas = AssetDataInterface.GetAllReferencedAssetDatas();
+		for (const UPlacementPaletteClient* Item : GetActivePaletteItems())
+		{
+			if (!Item)
+			{
+				continue;
+			}
+
+			if (ReferencedAssetDatas.FindByPredicate([Item](const FAssetData& ReferencedAssetData) { return (ReferencedAssetData.ToSoftObjectPath() == Item->AssetPath); }))
+			{
+				return true;
+			}
+
+			// The current implementation of the asset data interface for actors requires that individual actors report on assets contained within their components.
+			// Not all elements (legacy via actors) do this reliably, so additionally check the supplied factory for a match. 
+			if (!Item->FactoryInterfaceClass.Get())
+			{
+				continue;
+			}
+
+			TScriptInterface<IAssetFactoryInterface> FactoryInterface = GEditor->GetEditorSubsystem<UPlacementSubsystem>()->GetAssetFactoryFromFactoryClass(Item->FactoryInterfaceClass);
+			FAssetData FoundAssetDataFromFactory = FactoryInterface->GetAssetDataFromElementHandle(InElementToCheck);
+			if (FoundAssetDataFromFactory.ToSoftObjectPath() == Item->AssetPath)
+			{
+				return true;
+			}
+		}
+	}
+
+	return false;
+}
+
+const UPlacementPaletteAsset* UAssetPlacementSettings::GetActivePalette() const
+{
+	if (ActivePalette)
+	{
+		return ActivePalette;
+	}
+	return UserPalette;
+}
+
+UPlacementPaletteAsset* UAssetPlacementSettings::GetMutableActivePalette()
+{
+	if (ActivePalette)
+	{
+		return ActivePalette;
+	}
+	return UserPalette;
+}
diff --git a/Engine/Plugins/Experimental/AssetPlacementEdMode/Source/AssetPlacementEdMode/Private/AssetPlacementSettings.h b/Engine/Plugins/Experimental/AssetPlacementEdMode/Source/AssetPlacementEdMode/Private/AssetPlacementSettings.h
new file mode 100644
--- /dev/null
+++ b/Engine/Plugins/Experimental/AssetPlacementEdMode/Source/AssetPlacementEdMode/Private/AssetPlacementSettings.h
@@ -0,0 +1,157 @@
+// Copyright Epic Games, Inc. All Rights Reserved.
+
+#pragma once
+
+#include "CoreMinimal.h"
+#include "FoliageType.h"
+#include "PlacementPaletteItem.h"
+
+#include "AssetPlacementSettings.generated.h"
+
+class UPlacementPaletteAsset;
+class UPlacementPaletteClient;
+
+UCLASS(config = EditorPerProjectUserSettings)
+class UAssetPlacementSettings : public UObject
+{
+	GENERATED_BODY()
+
+public:
+	/** Set to true if the tool should automatically align placement to an axis. */
+	UPROPERTY(config, EditAnywhere, Category = "Axis Alignment", meta = (InlineEditConditionToggle))
+	bool bAlignToNormal = true;
+
+	/** The axis which the tool should use for alignment calculations. This axis is positive unless Invert Normal Axis is set. */
+	UPROPERTY(config, EditAnywhere, Category = "Axis Alignment", meta = (DisplayName = "Align To Normal", EditCondition = "bAlignToNormal"))
+	TEnumAsByte<EAxis::Type> AxisToAlignWithNormal = EAxis::Type::Z;
+
+	/** Set to true if the tool should negate the alignment axis. */
+	UPROPERTY(config, EditAnywhere, Category = "Axis Alignment")
+	bool bInvertNormalAxis = false;
+
+	/** Set to true if the tool should generate a random rotation on the X axis */
+	UPROPERTY(config, EditAnywhere, Category = "Rotation", meta = (InlineEditConditionToggle))
+	bool bUseRandomRotationX = false;
+
+	/** The range in degrees at which the tool should generate a random rotation on the X axis. */
+	UPROPERTY(config, EditAnywhere, AdvancedDisplay, Category = "Rotation", meta = (EditCondition = "bUseRandomRotationX", UIMin = "0.0", UIMax = "180.0"))
+	FFloatInterval RandomRotationX = FFloatInterval(0.0f, 180.0f);
+
+	/** Set to true if the tool should generate a negatively signed rotation on the X axis. */
+	UPROPERTY(config, EditAnywhere, AdvancedDisplay, Category = "Rotation")
+	bool bAllowNegativeRotationX = true;
+
+	/** Set to true if the tool should generate a random rotation on the Y axis */
+	UPROPERTY(config, EditAnywhere, Category = "Rotation", meta = (InlineEditConditionToggle))
+	bool bUseRandomRotationY = false;
+
+	/** The range in degrees at which the tool should generate a random rotation on the Y axis. */
+	UPROPERTY(config, EditAnywhere, AdvancedDisplay, Category = "Rotation", meta = (EditCondition = "bUseRandomRotationY", UIMin = "0.0", UIMax = "180.0"))
+	FFloatInterval RandomRotationY = FFloatInterval(0.0f, 180.0f);
+
+	/** Set to true if the tool should generate a negatively signed rotation on the Y axis. */
+	UPROPERTY(config, EditAnywhere, AdvancedDisplay, Category = "Rotation")
+	bool bAllowNegativeRotationY = true;
+
+	/** Set to true if the tool should generate a random rotation on the Z axis */
+	UPROPERTY(config, EditAnywhere, Category = "Rotation", meta = (InlineEditConditionToggle))
+	bool bUseRandomRotationZ = true;
+
+	/** The range in degrees at which the tool should generate a random rotation on the Z axis. */
+	UPROPERTY(config, EditAnywhere, Category = "Rotation", meta = (EditCondition = "bUseRandomRotationZ", UIMin = "0.0", UIMax = "180.0"))
+	FFloatInterval RandomRotationZ = FFloatInterval(0.0f, 180.0f);
+	
+	/** Set to true if the tool should generate a negatively signed rotation on the Z axis. */
+	UPROPERTY(config, EditAnywhere, Category = "Rotation")
+	bool bAllowNegativeRotationZ = true;
+
+	/** Which axes scale should be applied to. Any locked axes will be scaled to 1.0f. */
+	UPROPERTY(config, EditAnywhere, Category = "Scale")
+	EFoliageScaling ScalingType = EFoliageScaling::Uniform;
+	
+	/** The min and max at which a random scale is generated, or a scale is applied. */
+	UPROPERTY(config, EditAnywhere, Category = "Scale", meta = (UIMin = ".01"))
+	FFloatInterval ScaleRange = FFloatInterval(.5f, 2.0f);
+
+	/** Set to true if the tool should randomly scale any generated assets */
+	UPROPERTY(config, EditAnywhere, Category = "Scale", meta = (InlineEditConditionToggle))
+	bool bUseRandomScale = true;
+
+	/** Set to true if the tool should generate a negatively signed scale. */
+	UPROPERTY(config, EditAnywhere, AdvancedDisplay, Category = "Scale")
+	bool bAllowNegativeScale = false;
+
+	/** The spawn location offset in relative space. */
+	UPROPERTY(config, EditAnywhere, Category = "Location")
+	FVector RelativeLocationOffset = FVector::ZeroVector;
+
+	/** If the relative location offset should match scale with the placed asset. */
+	UPROPERTY(config, EditAnywhere, Category = "Location")
+	bool bScaleRelativeLocationOffset = true;
+
+	/** The spawn rotation offset in world space. */
+	UPROPERTY(config, EditAnywhere, AdvancedDisplay, Category = "Location")
+	FVector WorldLocationOffset = FVector::ZeroVector;
+
+	/** If the world location offset should match scale with the placed asset. */
+	UPROPERTY(config, EditAnywhere, AdvancedDisplay, Category = "Location")
+	bool bScaleWorldLocationOffset = false;
+
+	/** If the tool should place on landscape surfaces. */
+	UPROPERTY(config, EditAnywhere, Category = "Place on Filters")
+	bool bLandscape = true;
+
+	/** If the tool should place on static mesh surfaces. */
+	UPROPERTY(config, EditAnywhere, Category = "Place on Filters")
+	bool bStaticMeshes = true;
+
+	/** If the tool should place on BSP surfaces. */
+	UPROPERTY(config, EditAnywhere, Category = "Place on Filters")
+	bool bBSP = true;
+
+	/** If the tool should place on foliage surfaces. */
+	UPROPERTY(config, EditAnywhere, Category = "Place on Filters")
+	bool bFoliage = false;
+
+	/** If the tool should place on translucent surfaces. */
+	UPROPERTY(config, EditAnywhere, Category = "Place on Filters")
+	bool bTranslucent = false;
+
+	UPROPERTY(config)
+	bool bUseContentBrowserSelection = true;
+
+	virtual bool CanEditChange(const FProperty* InProperty) const override;
+
+	void SetPaletteAsset(UPlacementPaletteAsset* InPaletteAsset);
+	UPlacementPaletteClient* AddClientToActivePalette(const FAssetData& InAssetData);
+	TArrayView<const TObjectPtr<UPlacementPaletteClient>> GetActivePaletteItems() const;
+	const FGuid GetActivePaletteGuid() const;
+	void ClearActivePaletteItems();
+	void SaveActivePalette();
+	FSoftObjectPath GetActivePalettePath() const;
+	void LoadSettings();
+	void SaveSettings();
+
+	/**
+	 * Verifies if the given element handle is supported by the current mode settings' palette.
+	 *
+	 * @returns true if the element can be placed by the mode.
+	 */
+	bool DoesActivePaletteSupportElement(const FTypedElementHandle& InElementToCheck) const;
+
+protected:
+	const UPlacementPaletteAsset* GetActivePalette() const;
+	UPlacementPaletteAsset* GetMutableActivePalette();
+
+	UPROPERTY(Transient)
+	TObjectPtr<UPlacementPaletteAsset> ActivePalette;
+
+	UPROPERTY(Transient)
+	TObjectPtr<UPlacementPaletteAsset> UserPalette;
+
+	UPROPERTY(config)
+	FGuid UserGridGuid;
+
+	UPROPERTY(config)
+	FSoftObjectPath LastActivePalettePath;
+};
diff --git a/Engine/Plugins/Experimental/AssetPlacementEdMode/Source/AssetPlacementEdMode/Private/AssetTypeActions_PlacementPalette.cpp b/Engine/Plugins/Experimental/AssetPlacementEdMode/Source/AssetPlacementEdMode/Private/AssetTypeActions_PlacementPalette.cpp
new file mode 100644
--- /dev/null
+++ b/Engine/Plugins/Experimental/AssetPlacementEdMode/Source/AssetPlacementEdMode/Private/AssetTypeActions_PlacementPalette.cpp
@@ -0,0 +1,30 @@
+// Copyright Epic Games, Inc. All Rights Reserved.
+
+#include "AssetTypeActions_PlacementPalette.h"
+
+#include "PlacementPaletteAsset.h"
+
+FAssetTypeActions_PlacementPalette::FAssetTypeActions_PlacementPalette(EAssetTypeCategories::Type InAssetCategoryBit)
+	: AssetCategoryBit(InAssetCategoryBit)
+{
+}
+
+FText FAssetTypeActions_PlacementPalette::GetName() const
+{
+	return NSLOCTEXT("AssetPlacementEdMode", "AssetTypeActions_PlacementPalette", "Placement Palette");
+}
+
+FColor FAssetTypeActions_PlacementPalette::GetTypeColor() const
+{
+	return FColor(222, 132, 255);
+}
+
+UClass* FAssetTypeActions_PlacementPalette::GetSupportedClass() const
+{
+	return UPlacementPaletteAsset::StaticClass();
+}
+
+uint32 FAssetTypeActions_PlacementPalette::GetCategories()
+{
+	return AssetCategoryBit;
+}
diff --git a/Engine/Plugins/Experimental/AssetPlacementEdMode/Source/AssetPlacementEdMode/Private/AssetTypeActions_PlacementPalette.h b/Engine/Plugins/Experimental/AssetPlacementEdMode/Source/AssetPlacementEdMode/Private/AssetTypeActions_PlacementPalette.h
new file mode 100644
--- /dev/null
+++ b/Engine/Plugins/Experimental/AssetPlacementEdMode/Source/AssetPlacementEdMode/Private/AssetTypeActions_PlacementPalette.h
@@ -0,0 +1,21 @@
+// Copyright Epic Games, Inc. All Rights Reserved.
+
+#pragma once
+
+#include "CoreMinimal.h"
+#include "AssetTypeActions_Base.h"
+
+class FAssetTypeActions_PlacementPalette : public FAssetTypeActions_Base
+{
+public:
+	FAssetTypeActions_PlacementPalette(EAssetTypeCategories::Type InAssetCategoryBit);
+
+	// IAssetTypeActions Implementation
+	virtual FText GetName() const override;
+	virtual FColor GetTypeColor() const override;
+	virtual UClass* GetSupportedClass() const override;
+	virtual uint32 GetCategories() override;
+
+private:
+	EAssetTypeCategories::Type AssetCategoryBit;
+};
diff --git a/Engine/Plugins/Experimental/AssetPlacementEdMode/Source/AssetPlacementEdMode/Private/Modes/PlacementModeSubsystem.cpp b/Engine/Plugins/Experimental/AssetPlacementEdMode/Source/AssetPlacementEdMode/Private/Modes/PlacementModeSubsystem.cpp
new file mode 100644
--- /dev/null
+++ b/Engine/Plugins/Experimental/AssetPlacementEdMode/Source/AssetPlacementEdMode/Private/Modes/PlacementModeSubsystem.cpp
@@ -0,0 +1,48 @@
+// Copyright Epic Games, Inc. All Rights Reserved.
+
+#include "Modes/PlacementModeSubsystem.h"
+
+#include "AssetPlacementSettings.h"
+#include "Subsystems/PlacementSubsystem.h"
+
+#include "Editor.h"
+
+void UPlacementModeSubsystem::Initialize(FSubsystemCollectionBase& Collection)
+{
+	Collection.InitializeDependency<UPlacementSubsystem>();
+
+	GEditor->GetEditorSubsystem<UPlacementSubsystem>()->OnPlacementFactoriesRegistered().AddUObject(this, &UPlacementModeSubsystem::LoadSettings);
+	FCoreDelegates::OnEnginePreExit.AddUObject(this, &UPlacementModeSubsystem::SaveSettings);
+}
+
+void UPlacementModeSubsystem::Deinitialize()
+{
+	GEditor->GetEditorSubsystem<UPlacementSubsystem>()->OnPlacementFactoriesRegistered().RemoveAll(this);
+	ModeSettings = nullptr;
+}
+
+const UAssetPlacementSettings* UPlacementModeSubsystem::GetModeSettingsObject() const
+{
+	return ModeSettings;
+}
+
+UAssetPlacementSettings* UPlacementModeSubsystem::GetMutableModeSettingsObject()
+{
+	return ModeSettings;
+}
+
+void UPlacementModeSubsystem::SaveSettings() const
+{
+	if (ModeSettings)
+	{
+		ModeSettings->SaveSettings();
+	}
+}
+
+void UPlacementModeSubsystem::LoadSettings()
+{
+	ModeSettings = NewObject<UAssetPlacementSettings>(this);
+	ModeSettings->LoadSettings();
+
+	GEditor->GetEditorSubsystem<UPlacementSubsystem>()->OnPlacementFactoriesRegistered().RemoveAll(this);
+}
diff --git a/Engine/Plugins/Experimental/AssetPlacementEdMode/Source/AssetPlacementEdMode/Private/Modes/PlacementModeSubsystem.h b/Engine/Plugins/Experimental/AssetPlacementEdMode/Source/AssetPlacementEdMode/Private/Modes/PlacementModeSubsystem.h
new file mode 100644
--- /dev/null
+++ b/Engine/Plugins/Experimental/AssetPlacementEdMode/Source/AssetPlacementEdMode/Private/Modes/PlacementModeSubsystem.h
@@ -0,0 +1,33 @@
+// Copyright Epic Games, Inc. All Rights Reserved.
+
+#pragma once
+
+#include "EditorSubsystem.h"
+#include "UObject/ObjectPtr.h"
+
+#include "PlacementModeSubsystem.generated.h"
+
+class UAssetPlacementSettings;
+
+UCLASS()
+class UPlacementModeSubsystem : public UEditorSubsystem
+{
+	GENERATED_BODY()
+
+public:
+	// USubsystem Interface
+	virtual void Initialize(FSubsystemCollectionBase& Collection) override;
+	virtual void Deinitialize() override;
+	// End USubsystem Interface
+
+	// @returns the settings object for the mode for sharing across all tools and tool builders.
+	const UAssetPlacementSettings* GetModeSettingsObject() const;
+	UAssetPlacementSettings* GetMutableModeSettingsObject();
+
+protected:
+	UPROPERTY()
+	TObjectPtr<UAssetPlacementSettings> ModeSettings;
+
+	void SaveSettings() const;
+	void LoadSettings();
+};
diff --git a/Engine/Plugins/Experimental/AssetPlacementEdMode/Source/AssetPlacementEdMode/Private/PlacementPaletteAsset.cpp b/Engine/Plugins/Experimental/AssetPlacementEdMode/Source/AssetPlacementEdMode/Private/PlacementPaletteAsset.cpp
new file mode 100644
--- /dev/null
+++ b/Engine/Plugins/Experimental/AssetPlacementEdMode/Source/AssetPlacementEdMode/Private/PlacementPaletteAsset.cpp
@@ -0,0 +1,31 @@
+// Copyright Epic Games, Inc. All Rights Reserved.
+
+#include "PlacementPaletteAsset.h"
+
+UPlacementPaletteAssetFactory::UPlacementPaletteAssetFactory(const FObjectInitializer& ObjectInitializer)
+	: Super(ObjectInitializer)
+{
+	SupportedClass = UPlacementPaletteAsset::StaticClass();
+	bCreateNew = true;
+	bEditAfterNew = true;
+}
+
+UObject* UPlacementPaletteAssetFactory::FactoryCreateNew(UClass* InClass, UObject* InParent, FName InName, EObjectFlags Flags, UObject* Context, FFeedbackContext* Warn)
+{
+	UPlacementPaletteAsset* PlacementPaletteAsset = NewObject<UPlacementPaletteAsset>(InParent, InClass, InName, Flags, Context);
+	if (PlacementPaletteAsset)
+	{
+		PlacementPaletteAsset->GridGuid = FGuid::NewGuid();
+	}
+	return PlacementPaletteAsset;
+}
+
+bool UPlacementPaletteAssetFactory::ShouldShowInNewMenu() const
+{
+	return true;
+}
+
+FText UPlacementPaletteAssetFactory::GetToolTip() const
+{
+	return NSLOCTEXT("AssetPlacementEdMode", "PlacementPaletteToolTip", "Placement palettes can be used in the placement mode to quickly place many different asset types in a level with preconfigured settings.");
+}
diff --git a/Engine/Plugins/Experimental/AssetPlacementEdMode/Source/AssetPlacementEdMode/Private/PlacementPaletteAsset.h b/Engine/Plugins/Experimental/AssetPlacementEdMode/Source/AssetPlacementEdMode/Private/PlacementPaletteAsset.h
new file mode 100644
--- /dev/null
+++ b/Engine/Plugins/Experimental/AssetPlacementEdMode/Source/AssetPlacementEdMode/Private/PlacementPaletteAsset.h
@@ -0,0 +1,35 @@
+// Copyright Epic Games, Inc. All Rights Reserved.
+
+#pragma once
+
+#include "UObject/Object.h"
+#include "Factories/Factory.h"
+#include "PlacementPaletteItem.h"
+
+#include "PlacementPaletteAsset.generated.h"
+
+class UPlacementPaletteAsset;
+
+UCLASS(NotPlaceable)
+class UPlacementPaletteAsset : public UObject
+{
+	GENERATED_BODY()
+
+public:
+	UPROPERTY(EditAnywhere, Category=Items)
+	TArray<TObjectPtr<UPlacementPaletteClient>> PaletteItems;
+
+	UPROPERTY()
+	FGuid GridGuid;
+};
+
+UCLASS(hideCategories=Object)
+class UPlacementPaletteAssetFactory : public UFactory
+{
+	GENERATED_UCLASS_BODY()
+
+public:
+	virtual UObject* FactoryCreateNew(UClass* InClass, UObject* InParent, FName InName, EObjectFlags Flags, UObject* Context, FFeedbackContext* Warn) override;
+	virtual bool ShouldShowInNewMenu() const override;
+	virtual FText GetToolTip() const override;
+};
diff --git a/Engine/Plugins/Experimental/AssetPlacementEdMode/Source/AssetPlacementEdMode/Private/PlacementPaletteItem.cpp b/Engine/Plugins/Experimental/AssetPlacementEdMode/Source/AssetPlacementEdMode/Private/PlacementPaletteItem.cpp
new file mode 100644
--- /dev/null
+++ b/Engine/Plugins/Experimental/AssetPlacementEdMode/Source/AssetPlacementEdMode/Private/PlacementPaletteItem.cpp
@@ -0,0 +1,21 @@
+// Copyright Epic Games, Inc. All Rights Reserved.
+
+#include "PlacementPaletteItem.h"
+
+#include "Factories/AssetFactoryInterface.h"
+#include "Subsystems/PlacementSubsystem.h"
+
+void UPlacementPaletteClient::PostEditChangeProperty(struct FPropertyChangedEvent& PropertyChangedEvent)
+{
+	Super::PostEditChangeProperty(PropertyChangedEvent);
+}
+
+void UPlacementPaletteClient::PostLoad()
+{
+	Super::PostLoad();
+}
+
+void UPlacementPaletteClient::Serialize(FArchive& Ar)
+{
+	Super::Serialize(Ar);
+}
diff --git a/Engine/Plugins/Experimental/AssetPlacementEdMode/Source/AssetPlacementEdMode/Private/PlacementPaletteItem.h b/Engine/Plugins/Experimental/AssetPlacementEdMode/Source/AssetPlacementEdMode/Private/PlacementPaletteItem.h
new file mode 100644
--- /dev/null
+++ b/Engine/Plugins/Experimental/AssetPlacementEdMode/Source/AssetPlacementEdMode/Private/PlacementPaletteItem.h
@@ -0,0 +1,35 @@
+// Copyright Epic Games, Inc. All Rights Reserved.
+
+#pragma once
+
+#include "CoreMinimal.h"
+#include "ISMPartition/ISMPartitionInstanceManager.h"
+#include "Factories/AssetFactoryInterface.h"
+
+#include "PlacementPaletteItem.generated.h"
+
+class UAssetFactoryInterface;
+class UInstancedPlacemenClientSettings;
+
+UCLASS(NotPlaceable)
+class UPlacementPaletteClient : public UObject
+{
+	GENERATED_BODY()
+
+public:
+	UPROPERTY(EditAnywhere, Category = Settings)
+	FSoftObjectPath AssetPath;
+
+	UPROPERTY()
+	FGuid ClientGuid;
+
+	UPROPERTY(EditAnywhere, Category = Settings)
+	TObjectPtr<UInstancedPlacemenClientSettings> SettingsObject = nullptr;
+
+	UPROPERTY()
+	TSubclassOf<UAssetFactoryInterface> FactoryInterfaceClass;
+
+	virtual void PostEditChangeProperty(struct FPropertyChangedEvent& PropertyChangedEvent) override;
+	virtual void PostLoad() override;
+	virtual void Serialize(FArchive& Ar) override;
+};
diff --git a/Engine/Plugins/Experimental/AssetPlacementEdMode/Source/AssetPlacementEdMode/Private/SAssetPlacementPalette.cpp b/Engine/Plugins/Experimental/AssetPlacementEdMode/Source/AssetPlacementEdMode/Private/SAssetPlacementPalette.cpp
new file mode 100644
--- /dev/null
+++ b/Engine/Plugins/Experimental/AssetPlacementEdMode/Source/AssetPlacementEdMode/Private/SAssetPlacementPalette.cpp
@@ -0,0 +1,870 @@
+// Copyright Epic Games, Inc. All Rights Reserved.
+
+#include "SAssetPlacementPalette.h"
+#include "AssetPlacementPaletteItem.h"
+#include "Framework/Application/SlateApplication.h"
+#include "Framework/Commands/UIAction.h"
+#include "Framework/Commands/UICommandList.h"
+#include "Framework/MultiBox/MultiBoxBuilder.h"
+#include "Widgets/Images/SImage.h"
+#include "Widgets/Layout/SScrollBorder.h"
+#include "Widgets/Input/SSlider.h"
+#include "Widgets/Input/SCheckBox.h"
+#include "EditorStyleSet.h"
+#include "AssetThumbnail.h"
+
+#include "IContentBrowserSingleton.h"
+#include "ContentBrowserModule.h"
+#include "PropertyCustomizationHelpers.h"
+#include "AssetData.h"
+#include "AssetSelection.h"
+#include "Editor.h"
+#include "Widgets/Layout/SScrollBox.h"
+
+#include "Widgets/Layout/SScaleBox.h"
+#include "Widgets/Layout/SWidgetSwitcher.h"
+#include "Widgets/Input/SSearchBox.h"
+#include "AssetPlacementEdModeStyle.h"
+#include "AssetPlacementSettings.h"
+#include "Subsystems/PlacementSubsystem.h"
+#include "Engine/AssetManager.h"
+#include "Engine/StreamableManager.h"
+#include "Modes/PlacementModeSubsystem.h"
+
+#include "IContentBrowserSingleton.h"
+#include "ContentBrowserModule.h"
+#include "Modules/ModuleManager.h"
+#include "SAssetDropTarget.h"
+#include "PlacementPaletteAsset.h"
+#include "AssetRegistry/AssetRegistryModule.h"
+#include "PlacementPaletteItem.h"
+#include "Factories/AssetFactoryInterface.h"
+#include "Instances/EditorPlacementSettings.h"
+
+#define LOCTEXT_NAMESPACE "AssetPlacementMode"
+
+////////////////////////////////////////////////
+// SPlacementPalette
+////////////////////////////////////////////////
+void SAssetPlacementPalette::Construct(const FArguments& InArgs)
+{
+	ItemDetailsWidget = InArgs._ItemDetailsView;
+	bItemsNeedRebuild = false;
+	bIsRebuildTimerRegistered = false;
+	bIsRefreshTimerRegistered = false;
+
+	UICommandList = MakeShared<FUICommandList>();
+
+	// Size of the thumbnail pool should be large enough to show a reasonable amount of Placement assets on screen at once,
+	// otherwise some thumbnail images will appear duplicated.
+	ThumbnailPool = MakeShared<FAssetThumbnailPool>(64);
+
+	TypeFilter = MakeShared<PlacementTypeTextFilter>(PlacementTypeTextFilter::FItemToStringArray::CreateSP(this, &SAssetPlacementPalette::GetPaletteItemFilterString));
+
+	const FText BlankText = FText::GetEmpty();
+
+	// Load initial settings for the palette widget to use
+	if (UPlacementModeSubsystem* PlacementModeSubystem = GEditor->GetEditorSubsystem<UPlacementModeSubsystem>())
+	{
+		if (const UAssetPlacementSettings* PlacementSettings = PlacementModeSubystem->GetModeSettingsObject())
+		{
+			bIsMirroringContentBrowser = !PlacementSettings->bUseContentBrowserSelection;
+			OnSetPaletteAsset(PlacementSettings->GetActivePalettePath().TryLoad());
+		}
+	}
+
+	ChildSlot
+	[
+		SNew(SVerticalBox)
+		+ SVerticalBox::Slot()
+		.HAlign(HAlign_Fill)
+		.AutoHeight()
+		[
+			// Top bar
+			SNew(SBorder)
+			.BorderImage(FAppStyle::Get().GetBrush("DetailsView.CategoryTop"))
+			.BorderBackgroundColor(FLinearColor(.6f, .6f, .6f, 1.0f))
+			[
+				SNew(SVerticalBox)
+				+ SVerticalBox::Slot()
+				.AutoHeight()
+				.HAlign(HAlign_Fill)
+				[
+					SNew(SCheckBox)
+					.Type(ESlateCheckBoxType::Type::ToggleButton)
+					.IsChecked(bIsMirroringContentBrowser ? ECheckBoxState::Checked : ECheckBoxState::Unchecked)
+					.Style(&FAssetPlacementEdModeStyle::Get().GetWidgetStyle<FCheckBoxStyle>("ToggleButtonCheckBox"))
+					.Visibility(this, &SAssetPlacementPalette::GetContentBrowserMirrorVisibility)
+					.OnCheckStateChanged(this, &SAssetPlacementPalette::OnContentBrowserMirrorButtonClicked)
+					[
+						SNew(STextBlock)
+						.Text(LOCTEXT("Placement_ToggleContentBrowserMirroring", "Mirror Content Browser Selection"))
+						.Justification(ETextJustify::Type::Center)
+						.TextStyle(&FAssetPlacementEdModeStyle::Get().GetWidgetStyle<FTextBlockStyle>("ButtonText"))
+						.ToolTipText(LOCTEXT("Placement_ToggleContentBrowserMirroring_ToolTip", "Toggles palette to mirror the active content browser selection."))
+					]
+				]
+				+SVerticalBox::Slot()
+				.AutoHeight()
+				[
+					SNew(SObjectPropertyEntryBox)
+					.AllowedClass(UPlacementPaletteAsset::StaticClass())
+					.ObjectPath(this, &SAssetPlacementPalette::GetPalettePath)
+					.OnObjectChanged(this, &SAssetPlacementPalette::OnSetPaletteAsset)
+					.ThumbnailPool(ThumbnailPool)
+					.CustomContentSlot()
+					[
+						SNew(SBox)
+						.HAlign(HAlign_Left)
+						.VAlign(VAlign_Center)
+						[
+							SNew(SHorizontalBox)
+							+ SHorizontalBox::Slot()
+							.VAlign(VAlign_Center)
+							.AutoWidth()
+							[
+								PropertyCustomizationHelpers::MakeClearButton(FSimpleDelegate::CreateSP(this, &SAssetPlacementPalette::OnResetPaletteAssetClicked), LOCTEXT("ResetPaletteAssetTooltip", "Clear the current palette asset, and use the user's local palette."), TAttribute<bool>::CreateLambda([this]() { return PalettePath.IsValid(); }))
+							]
+							+ SHorizontalBox::Slot()
+							.VAlign(VAlign_Center)
+							.AutoWidth()
+							.Padding(2.0f, 0.0f)
+							[
+								PropertyCustomizationHelpers::MakeSaveButton(FSimpleDelegate::CreateSP(this, &SAssetPlacementPalette::OnSavePaletteAssetClicked), LOCTEXT("SaveAssetPaletteTooltip", "Save changes to the current palette asset"), TAttribute<bool>::CreateLambda([this]() { return PalettePath.IsValid(); }))
+							]
+						]
+					]
+				]
+
+				+ SVerticalBox::Slot()
+				.VAlign(VAlign_Center)
+				.AutoHeight()
+				[
+					SNew(SHorizontalBox)
+
+					+ SHorizontalBox::Slot()
+					.FillWidth(.75f)
+					[
+						SAssignNew(SearchBoxPtr, SSearchBox)
+						.HintText(LOCTEXT("SearchPlacementPaletteHint", "Search Palette"))
+						.OnTextChanged(this, &SAssetPlacementPalette::OnSearchTextChanged)
+					]
+
+					// View Options
+					+ SHorizontalBox::Slot()
+					.AutoWidth()
+					[
+						SNew(SComboButton)
+						.ForegroundColor(FSlateColor::UseForeground())
+						.ButtonStyle(FAppStyle::Get(), "ToggleButton")
+						.OnGetMenuContent(this, &SAssetPlacementPalette::GetViewOptionsMenuContent)
+						.ButtonContent()
+						[
+							SNew(SBox)
+							.VAlign(VAlign_Center)
+							[
+								SNew(SImage)
+								.Image(FAppStyle::Get().GetBrush("GenericViewButton"))
+							]
+						]
+					]
+				]
+			]
+		]
+
+		+ SVerticalBox::Slot()
+		.HAlign(HAlign_Fill)
+		.VAlign(VAlign_Fill)
+		[
+			SNew(SAssetDropTarget)
+			.OnAreAssetsAcceptableForDrop(this, &SAssetPlacementPalette::OnAreAssetsValidForDrop)
+			.OnAssetsDropped(this, &SAssetPlacementPalette::HandlePlacementDropped)
+			.bSupportsMultiDrop(true)
+			[
+				SNew(SOverlay)
+				+ SOverlay::Slot()
+				[
+					CreatePaletteViews()
+				]
+
+				+ SOverlay::Slot()
+				.HAlign(HAlign_Center)
+				.VAlign(VAlign_Center)
+				[
+					SNew(STextBlock)
+					.Visibility(this, &SAssetPlacementPalette::GetDropHintTextVisibility)
+					.Text(LOCTEXT("Placement_AddPlacementMesh", "Drop Assets Here"))
+					.TextStyle(FAppStyle::Get(), "HintText")
+				]
+			]
+		]
+	];
+
+	UpdatePalette(true);
+}
+
+SAssetPlacementPalette::~SAssetPlacementPalette()
+{
+}
+
+FReply SAssetPlacementPalette::OnKeyDown(const FGeometry& MyGeometry, const FKeyEvent& InKeyEvent)
+{
+	if (UICommandList->ProcessCommandBindings(InKeyEvent))
+	{
+		return FReply::Handled();
+	}
+
+	return SCompoundWidget::OnKeyDown(MyGeometry, InKeyEvent);
+}
+
+void SAssetPlacementPalette::UpdatePalette(bool bRebuildItems)
+{
+	bItemsNeedRebuild |= bRebuildItems;
+
+	if (!bIsRebuildTimerRegistered)
+	{
+		bIsRebuildTimerRegistered = true;
+		RegisterActiveTimer(0.f, FWidgetActiveTimerDelegate::CreateSP(this, &SAssetPlacementPalette::UpdatePaletteItems));
+	}
+}
+
+void SAssetPlacementPalette::RefreshPalette()
+{
+	// Do not register the refresh timer if we're pending a rebuild; rebuild should cause the palette to refresh
+	if (!bIsRefreshTimerRegistered && !bIsRebuildTimerRegistered)
+	{
+		bIsRefreshTimerRegistered = true;
+		RegisterActiveTimer(0.f, FWidgetActiveTimerDelegate::CreateSP(this, &SAssetPlacementPalette::RefreshPaletteItems));
+	}
+}
+
+void SAssetPlacementPalette::RefreshActivePaletteViewWidget()
+{
+	if (ActiveViewMode == EViewMode::Thumbnail)
+	{
+		TileViewWidget->RequestListRefresh();
+	}
+	else
+	{
+		TreeViewWidget->RequestTreeRefresh();
+	}
+}
+
+void SAssetPlacementPalette::AddPlacementType(const FAssetData& AssetData)
+{
+	// Try to add the item to the mode's palette
+	if (UPlacementModeSubsystem* PlacementModeSubsystem = GEditor->GetEditorSubsystem<UPlacementModeSubsystem>())
+	{
+		UPlacementPaletteClient* PlacementInfo = PlacementModeSubsystem->GetMutableModeSettingsObject()->AddClientToActivePalette(AssetData);
+		if (!PlacementInfo || !PlacementInfo->AssetPath.IsValid())
+		{
+			return;
+		}
+
+		// Try to load the asset async so it's ready to place
+		UAssetManager::GetStreamableManager().RequestAsyncLoad(AssetData.ToSoftObjectPath());
+
+		PaletteItems.Add(MakeShared<FAssetPlacementPaletteItemModel>(AssetData, PlacementInfo, SharedThis(this), ThumbnailPool));
+	}
+}
+
+void SAssetPlacementPalette::ClearPalette()
+{
+	if (UPlacementModeSubsystem* ModeSubsystem = GEditor->GetEditorSubsystem<UPlacementModeSubsystem>())
+	{
+		ModeSubsystem->GetMutableModeSettingsObject()->ClearActivePaletteItems();
+	}
+
+	if (TSharedPtr<IDetailsView> PinnedItemDetails = ItemDetailsWidget.Pin())
+	{
+		PinnedItemDetails->SetObject(nullptr);
+	}
+
+	PaletteItems.Empty();
+	FilteredItems.Empty();
+}
+
+void SAssetPlacementPalette::OnClearPalette()
+{
+	ClearPalette();
+	UpdatePalette(true);
+}
+
+void SAssetPlacementPalette::SetPaletteToAssetDataList(TArrayView<const FAssetData> InAssetDatas)
+{
+	ClearPalette();
+	for (const FAssetData& SelectedAsset : InAssetDatas)
+	{
+		AddPlacementType(SelectedAsset);
+	}
+}
+
+TSharedRef<SWidgetSwitcher> SAssetPlacementPalette::CreatePaletteViews()
+{
+	const FText BlankText = FText::GetEmpty();
+
+	// Tile View Widget
+	SAssignNew(TileViewWidget, SPlacementTypeTileView)
+		.ListItemsSource(&FilteredItems)
+		.OnGenerateTile(this, &SAssetPlacementPalette::GenerateTile)
+		.OnContextMenuOpening(this, &SAssetPlacementPalette::ConstructPlacementTypeContextMenu)
+		.SelectionMode(ESelectionMode::Multi)
+		.OnSelectionChanged(this, &SAssetPlacementPalette::OnPaletteSelectionChanged)
+		.ItemHeight(this, &SAssetPlacementPalette::GetScaledThumbnailSize)
+		.ItemWidth(this, &SAssetPlacementPalette::GetScaledThumbnailSize)
+		.ItemAlignment(EListItemAlignment::LeftAligned);
+
+	// Tree View Widget
+	SAssignNew(TreeViewWidget, SPlacementTypeTreeView)
+		.TreeItemsSource(&FilteredItems)
+		.OnGenerateRow(this, &SAssetPlacementPalette::TreeViewGenerateRow)
+		.OnGetChildren(this, &SAssetPlacementPalette::TreeViewGetChildren)
+		.SelectionMode(ESelectionMode::Multi)
+		.OnSelectionChanged(this, &SAssetPlacementPalette::OnPaletteSelectionChanged)
+		.OnContextMenuOpening(this, &SAssetPlacementPalette::ConstructPlacementTypeContextMenu)
+		.HeaderRow
+		(
+			SAssignNew(TreeViewHeaderRow, SHeaderRow)
+			// Type
+			+ SHeaderRow::Column(AssetPlacementPaletteTreeColumns::ColumnID_Type)
+			.HeaderContentPadding(FMargin(10, 1, 0, 1))
+			.SortMode(this, &SAssetPlacementPalette::GetMeshColumnSortMode)
+			.OnSort(this, &SAssetPlacementPalette::OnTypeColumnSortModeChanged)
+			.DefaultLabel(this, &SAssetPlacementPalette::GetTypeColumnHeaderText)
+			.FillWidth(5.f)
+		);
+
+	// View Mode Switcher
+	SAssignNew(WidgetSwitcher, SWidgetSwitcher);
+
+	// Thumbnail View
+	WidgetSwitcher->AddSlot((uint8)EViewMode::Thumbnail)
+	[
+		SNew(SScrollBorder, TileViewWidget.ToSharedRef())
+		.Content()
+		[
+			TileViewWidget.ToSharedRef()
+		]
+	];
+
+	// Tree View
+	WidgetSwitcher->AddSlot((uint8)EViewMode::Tree)
+	[
+		SNew(SScrollBorder, TreeViewWidget.ToSharedRef())
+		.Style(&FAssetPlacementEdModeStyle::Get().GetWidgetStyle<FScrollBorderStyle>("FoliageEditMode.TreeView.ScrollBorder"))
+		.Content()
+		[
+			TreeViewWidget.ToSharedRef()
+		]
+	];
+
+	WidgetSwitcher->SetActiveWidgetIndex((uint8)ActiveViewMode);
+
+	return WidgetSwitcher.ToSharedRef();
+}
+
+void SAssetPlacementPalette::GetPaletteItemFilterString(FPlacementPaletteItemModelPtr PaletteItemModel, TArray<FString>& OutArray) const
+{
+	OutArray.Add(PaletteItemModel->GetDisplayFName().ToString());
+}
+
+void SAssetPlacementPalette::OnSearchTextChanged(const FText& InFilterText)
+{
+	TypeFilter->SetRawFilterText(InFilterText);
+	SearchBoxPtr->SetError(TypeFilter->GetFilterErrorText());
+	UpdatePalette();
+}
+
+bool SAssetPlacementPalette::ShouldFilterAsset(const FAssetData& InAssetData)
+{
+	UClass* Class = InAssetData.GetClass();
+
+	if (Class->HasAnyClassFlags(CLASS_Abstract | CLASS_Deprecated | CLASS_NewerVersionExists | CLASS_NotPlaceable))
+	{
+		return true;
+	}
+
+	if (UPlacementSubsystem* PlacementSubsystem = GEditor->GetEditorSubsystem<UPlacementSubsystem>())
+	{
+		return (PlacementSubsystem->FindAssetFactoryFromAssetData(InAssetData) != nullptr);
+	}
+
+	return true;
+}
+
+void SAssetPlacementPalette::OnResetPaletteAssetClicked()
+{
+	OnSetPaletteAsset(FAssetData());
+	UpdatePalette(true);
+}
+
+void SAssetPlacementPalette::OnSavePaletteAssetClicked()
+{
+	if (UPlacementModeSubsystem* PlacementModeSubsystem = GEditor->GetEditorSubsystem<UPlacementModeSubsystem>())
+	{
+		PlacementModeSubsystem->GetMutableModeSettingsObject()->SaveActivePalette();
+	}
+}
+
+EVisibility SAssetPlacementPalette::GetContentBrowserMirrorVisibility() const
+{
+	return GetPalettePath().IsEmpty() ? EVisibility::Visible : EVisibility::Collapsed;
+}
+
+void SAssetPlacementPalette::OnContentBrowserMirrorButtonClicked(ECheckBoxState InState)
+{
+	SetupContentBrowserMirroring((InState == ECheckBoxState::Checked));
+
+	if (UPlacementModeSubsystem* PlacementModeSubsystem = GEditor->GetEditorSubsystem<UPlacementModeSubsystem>())
+	{
+		PlacementModeSubsystem->GetMutableModeSettingsObject()->bUseContentBrowserSelection = bIsMirroringContentBrowser;
+	}
+}
+
+void SAssetPlacementPalette::OnContentBrowserSelectionChanged(const TArray<FAssetData>& NewSelectedAssets, bool bIsPrimaryBrowser)
+{
+	if (bIsPrimaryBrowser)
+	{
+		SetPaletteToAssetDataList(NewSelectedAssets);
+		UpdatePalette(true);
+	}
+}
+
+void SAssetPlacementPalette::SetupContentBrowserMirroring(bool bInMirrorContentBrowser)
+{
+	bool bWasMirroringContentBrowser = bIsMirroringContentBrowser;
+	if (bWasMirroringContentBrowser != bInMirrorContentBrowser)
+	{
+		if (FContentBrowserModule* ContentBrowserModule = FModuleManager::GetModulePtr<FContentBrowserModule>("ContentBrowser"))
+		{
+			if (bInMirrorContentBrowser)
+			{
+				TArray<FAssetData> SelectedAssetDatas;
+				ContentBrowserModule->Get().GetSelectedAssets(SelectedAssetDatas);
+				OnContentBrowserSelectionChanged(SelectedAssetDatas, true);
+				ContentBrowserModule->GetOnAssetSelectionChanged().AddSP(SharedThis(this), &SAssetPlacementPalette::OnContentBrowserSelectionChanged);
+			}
+			else
+			{
+				ContentBrowserModule->GetOnAssetSelectionChanged().RemoveAll(this);
+			}
+		}
+	}
+	bIsMirroringContentBrowser = bInMirrorContentBrowser;
+}
+
+FString SAssetPlacementPalette::GetPalettePath() const
+{
+	return PalettePath.ToString();
+}
+
+void SAssetPlacementPalette::OnSetPaletteAsset(const FAssetData& InAssetData)
+{
+	FSoftObjectPath NewAssetPath = InAssetData.ToSoftObjectPath();
+	if ((NewAssetPath != PalettePath) || !InAssetData.IsValid())
+	{
+		UPlacementPaletteAsset* PaletteAsset = Cast<UPlacementPaletteAsset>(InAssetData.GetAsset());
+		if (UPlacementModeSubsystem* PlacementModeSubsystem = GEditor->GetEditorSubsystem<UPlacementModeSubsystem>())
+		{
+			PlacementModeSubsystem->GetMutableModeSettingsObject()->SetPaletteAsset(PaletteAsset);
+
+			PalettePath = NewAssetPath;
+
+			// Setup content browser mirroring if this is the user palette
+			SetupContentBrowserMirroring(PalettePath.IsValid() ? false : PlacementModeSubsystem->GetModeSettingsObject()->bUseContentBrowserSelection);
+
+			SetPaletteItems(PlacementModeSubsystem->GetModeSettingsObject()->GetActivePaletteItems());
+		}
+	}
+}
+
+void SAssetPlacementPalette::SetPaletteItems(TArrayView<const TObjectPtr<UPlacementPaletteClient>> InPaletteItems)
+{
+	PaletteItems.Empty();
+	FilteredItems.Empty();
+	if (TSharedPtr<IDetailsView> PinnedDetailsView = ItemDetailsWidget.Pin())
+	{
+		PinnedDetailsView->SetObject(nullptr);
+	}
+
+	FAssetRegistryModule& AssetRegistry = FModuleManager::Get().LoadModuleChecked<FAssetRegistryModule>(TEXT("AssetRegistry"));
+	TArray<FSoftObjectPath> AssetsToLoad;
+	for (const UPlacementPaletteClient* PaletteItem : InPaletteItems)
+	{
+		if (PaletteItem)
+		{
+			FAssetData AssetData = AssetRegistry.Get().GetAssetByObjectPath(PaletteItem->AssetPath.GetAssetPathName());
+			if (AssetData.IsValid())
+			{
+				PaletteItems.Add(MakeShared<FAssetPlacementPaletteItemModel>(AssetData, PaletteItem, SharedThis(this), ThumbnailPool));
+				AssetsToLoad.Emplace(PaletteItem->AssetPath);
+			}
+		}
+	}
+
+	// Try to load the assets async so they're ready to place later
+	UAssetManager::GetStreamableManager().RequestAsyncLoad(AssetsToLoad);
+
+	// Update the palette view to match contained items
+	UpdatePalette(true);
+}
+
+void SAssetPlacementPalette::SetViewMode(EViewMode NewViewMode)
+{
+	if (ActiveViewMode != NewViewMode)
+	{
+		ActiveViewMode = NewViewMode;
+		WidgetSwitcher->SetActiveWidgetIndex((uint8)ActiveViewMode);
+		
+		RefreshActivePaletteViewWidget();
+	}
+}
+
+bool SAssetPlacementPalette::IsActiveViewMode(EViewMode ViewMode) const
+{
+	return ActiveViewMode == ViewMode;
+}
+
+void SAssetPlacementPalette::ToggleShowTooltips()
+{
+	bShowFullTooltips = !bShowFullTooltips;
+}
+
+bool SAssetPlacementPalette::ShouldShowTooltips() const
+{
+	return bShowFullTooltips;
+}
+
+FText SAssetPlacementPalette::GetSearchText() const
+{
+	return TypeFilter->GetRawFilterText();
+}
+
+TSharedRef<SWidget> SAssetPlacementPalette::GetViewOptionsMenuContent()
+{
+	FMenuBuilder MenuBuilder(true, UICommandList);
+
+	MenuBuilder.BeginSection("PlacementPaletteViewMode", LOCTEXT("ViewModeHeading", "Palette View Mode"));
+	{
+		MenuBuilder.AddMenuEntry(
+			LOCTEXT("ThumbnailView", "Thumbnails"),
+			LOCTEXT("ThumbnailView_ToolTip", "Display thumbnails for each Placement type in the palette."),
+			FSlateIcon(),
+			FUIAction(
+				FExecuteAction::CreateSP(this, &SAssetPlacementPalette::SetViewMode, EViewMode::Thumbnail),
+				FCanExecuteAction(),
+				FIsActionChecked::CreateSP(this, &SAssetPlacementPalette::IsActiveViewMode, EViewMode::Thumbnail)
+				),
+			NAME_None,
+			EUserInterfaceActionType::ToggleButton
+			);
+
+		MenuBuilder.AddMenuEntry(
+			LOCTEXT("ListView", "List"),
+			LOCTEXT("ListView_ToolTip", "Display Placement types in the palette as a list."),
+			FSlateIcon(),
+			FUIAction(
+				FExecuteAction::CreateSP(this, &SAssetPlacementPalette::SetViewMode, EViewMode::Tree),
+				FCanExecuteAction(),
+				FIsActionChecked::CreateSP(this, &SAssetPlacementPalette::IsActiveViewMode, EViewMode::Tree)
+				),
+			NAME_None,
+			EUserInterfaceActionType::ToggleButton
+			);
+	}
+	MenuBuilder.EndSection();
+
+	MenuBuilder.BeginSection("PlacementPaletteViewOptions", LOCTEXT("ViewOptionsHeading", "View Options"));
+	{
+		MenuBuilder.AddMenuEntry(
+			LOCTEXT("ShowTooltips", "Show Tooltips"),
+			LOCTEXT("ShowTooltips_ToolTip", "Whether to show tooltips when hovering over Placement types in the palette."),
+			FSlateIcon(),
+			FUIAction(
+				FExecuteAction::CreateSP(this, &SAssetPlacementPalette::ToggleShowTooltips),
+				FCanExecuteAction(),
+				FIsActionChecked::CreateSP(this, &SAssetPlacementPalette::ShouldShowTooltips),
+				FIsActionButtonVisible::CreateSP(this, &SAssetPlacementPalette::IsActiveViewMode, EViewMode::Tree)
+				),
+			NAME_None,
+			EUserInterfaceActionType::ToggleButton
+			);
+
+		MenuBuilder.AddWidget(
+			SNew(SHorizontalBox)
+			.Visibility(this, &SAssetPlacementPalette::GetThumbnailScaleSliderVisibility)
+			+ SHorizontalBox::Slot()
+			[
+				SNew(STextBlock)
+				.Text(LOCTEXT("ThumbnailScaleLabel", "Scale"))
+			]
+			+SHorizontalBox::Slot()
+			[
+				SNew(SSlider)
+				.ToolTipText(LOCTEXT("ThumbnailScaleToolTip", "Adjust the size of thumbnails."))
+				.Value(this, &SAssetPlacementPalette::GetThumbnailScale)
+				.OnValueChanged(this, &SAssetPlacementPalette::SetThumbnailScale)
+				.OnMouseCaptureEnd(this, &SAssetPlacementPalette::RefreshActivePaletteViewWidget)
+			],
+			FText(),
+			/*bNoIndent=*/true
+			);
+	}
+	MenuBuilder.EndSection();
+
+	return MenuBuilder.MakeWidget();
+}
+
+TSharedPtr<SListView<FPlacementPaletteItemModelPtr>> SAssetPlacementPalette::GetActiveViewWidget() const
+{
+	if (ActiveViewMode == EViewMode::Thumbnail)
+	{
+		return TileViewWidget;
+	}
+	else if (ActiveViewMode == EViewMode::Tree)
+	{
+		return TreeViewWidget;
+	}
+	
+	return nullptr;
+}
+
+EVisibility SAssetPlacementPalette::GetDropHintTextVisibility() const
+{
+	if (!bIsMirroringContentBrowser && (FilteredItems.Num() == 0) && !FSlateApplication::Get().IsDragDropping())
+	{
+		return EVisibility::HitTestInvisible;
+	}
+
+	return EVisibility::Hidden;
+}
+
+bool SAssetPlacementPalette::OnAreAssetsValidForDrop(TArrayView<FAssetData> DraggedAssets) const
+{
+	if (bIsMirroringContentBrowser)
+	{
+		return false;
+	}
+
+	UPlacementSubsystem* PlacementSubystem = GEditor->GetEditorSubsystem<UPlacementSubsystem>();
+	for (const FAssetData& AssetData : DraggedAssets)
+	{
+		if (AssetData.GetClass()->HasAnyClassFlags(CLASS_Abstract | CLASS_Deprecated | CLASS_NewerVersionExists | CLASS_NotPlaceable))
+		{
+			continue;
+		}
+
+		if (AssetData.IsValid())
+		{
+			return PlacementSubystem ? (PlacementSubystem->FindAssetFactoryFromAssetData(AssetData).GetInterface() != nullptr) : true;
+		}
+	}
+
+	return false;
+}
+
+void SAssetPlacementPalette::HandlePlacementDropped(const FDragDropEvent& DragDropEvent, TArrayView<FAssetData> DroppedAssetData)
+{
+	if ((DroppedAssetData.Num() == 0) || bIsMirroringContentBrowser)
+	{
+		return;
+	}
+
+	if (DragDropEvent.IsShiftDown())
+	{
+		ClearPalette();
+	}
+
+	for (auto& AssetData : DroppedAssetData)
+	{
+		AddPlacementType(AssetData);
+	}
+
+	UpdatePalette(true);
+}
+
+bool SAssetPlacementPalette::HasAnyItemInPalette() const
+{
+	return PaletteItems.Num() > 0;
+}
+
+TSharedPtr<SWidget> SAssetPlacementPalette::ConstructPlacementTypeContextMenu()
+{
+	FMenuBuilder MenuBuilder(true, nullptr);
+	if (!bIsMirroringContentBrowser)
+	{
+		MenuBuilder.AddMenuEntry(
+			LOCTEXT("Palette_Clear", "Clear Palette"),
+			LOCTEXT("Palette_ClearDesc", "Removes all items from the palette."),
+			FSlateIcon(),
+			FUIAction(
+				FExecuteAction::CreateSP(this, &SAssetPlacementPalette::OnClearPalette),
+				FCanExecuteAction::CreateSP(this, &SAssetPlacementPalette::HasAnyItemInPalette))
+		);
+	}
+	return MenuBuilder.MakeWidget();
+}
+
+void SAssetPlacementPalette::OnPaletteSelectionChanged(FPlacementPaletteItemModelPtr Item, ESelectInfo::Type SelectInfo)
+{
+	if (SelectInfo == ESelectInfo::Direct)
+	{
+		return;
+	}
+
+	// Sync selection to the widgets which are not active
+	switch (ActiveViewMode)
+	{
+		case EViewMode::Thumbnail:
+			TreeViewWidget->SetItemSelection(Item, TileViewWidget->IsItemSelected(Item), ESelectInfo::Direct);
+		break;
+
+		case EViewMode::Tree:
+			TileViewWidget->SetItemSelection(Item, TreeViewWidget->IsItemSelected(Item), ESelectInfo::Direct);
+		break;
+	}
+
+	if (TSharedPtr<IDetailsView> PinnedItemDetails = ItemDetailsWidget.Pin())
+	{
+		TArray<UObject*> DetailsObjects;
+		for (const FPlacementPaletteItemModelPtr& SelectedItem : GetActiveViewWidget()->GetSelectedItems())
+		{
+			if (UInstancedPlacemenClientSettings* SelectedItemSettingsObject = SelectedItem->GetTypeUIInfo()->SettingsObject.Get())
+			{
+				DetailsObjects.Add(SelectedItemSettingsObject);
+			}
+		}
+
+		constexpr bool bForceRefresh = true;
+		PinnedItemDetails->SetObjects(DetailsObjects, true);
+	}
+}
+
+void SAssetPlacementPalette::OnShowPlacementTypeInCB()
+{
+	TArray<FAssetData> FilteredAssets;
+	for (FPlacementPaletteItemModelPtr& PaletteItem : FilteredItems)
+	{
+		if (PaletteItem.IsValid() && PaletteItem->GetTypeUIInfo().IsValid())
+		{
+			FilteredAssets.Add(PaletteItem->GetTypeUIInfo()->AssetData);
+		}
+	}
+
+	if (FilteredAssets.Num())
+	{
+		GEditor->SyncBrowserToObjects(FilteredAssets);
+	}
+}
+
+// THUMBNAIL VIEW
+
+TSharedRef<ITableRow> SAssetPlacementPalette::GenerateTile(FPlacementPaletteItemModelPtr Item, const TSharedRef<STableViewBase>& OwnerTable)
+{
+	return SNew(SAssetPlacementPaletteItemTile, OwnerTable, Item);
+
+	// Refresh the palette to ensure that thumbnails are correct
+	RefreshPalette();
+}
+
+float SAssetPlacementPalette::GetScaledThumbnailSize() const
+{
+	const FInt32Interval& SizeRange = PlacementPaletteConstants::ThumbnailSizeRange;
+	return SizeRange.Min + SizeRange.Size() * GetThumbnailScale();
+}
+
+float SAssetPlacementPalette::GetThumbnailScale() const
+{
+	return PaletteThumbnailScale;
+}
+
+void SAssetPlacementPalette::SetThumbnailScale(float InScale)
+{
+	PaletteThumbnailScale = FMath::Clamp(InScale, 0.f, 1.f);
+}
+
+EVisibility SAssetPlacementPalette::GetThumbnailScaleSliderVisibility() const
+{
+	return (ActiveViewMode == EViewMode::Thumbnail) ? EVisibility::Visible : EVisibility::Collapsed;
+}
+
+// TREE VIEW
+
+TSharedRef<ITableRow> SAssetPlacementPalette::TreeViewGenerateRow(FPlacementPaletteItemModelPtr Item, const TSharedRef<STableViewBase>& OwnerTable)
+{
+	return SNew(SAssetPlacementPaletteItemRow, OwnerTable, Item);
+}
+
+void SAssetPlacementPalette::TreeViewGetChildren(FPlacementPaletteItemModelPtr Item, TArray<FPlacementPaletteItemModelPtr>& OutChildren)
+{
+	// Items do not have any children
+}
+
+FText SAssetPlacementPalette::GetTypeColumnHeaderText() const
+{
+	return LOCTEXT("PlacementTypeHeader", "Asset Type");
+}
+
+EColumnSortMode::Type SAssetPlacementPalette::GetMeshColumnSortMode() const
+{
+	return ActiveSortOrder;
+}
+
+void SAssetPlacementPalette::OnTypeColumnSortModeChanged(EColumnSortPriority::Type InPriority, const FName& InColumnName, EColumnSortMode::Type InSortMode)
+{
+	if (ActiveSortOrder == InSortMode)
+	{
+		return;
+	}
+
+	ActiveSortOrder = InSortMode;
+
+	if (ActiveSortOrder != EColumnSortMode::None)
+	{
+		auto CompareEntry = [this](const FPlacementPaletteItemModelPtr& A, const FPlacementPaletteItemModelPtr& B)
+		{
+			bool CompareResult = (A->GetDisplayFName().GetComparisonIndex().CompareLexical(B->GetDisplayFName().GetComparisonIndex()) <= 0);
+			return (ActiveSortOrder == EColumnSortMode::Ascending) ? CompareResult : !CompareResult;
+		};
+
+		PaletteItems.Sort(CompareEntry);
+	}
+}
+
+EActiveTimerReturnType SAssetPlacementPalette::UpdatePaletteItems(double InCurrentTime, float InDeltaTime)
+{
+	if (bItemsNeedRebuild)
+	{
+		bItemsNeedRebuild = false;
+	}
+
+	// Update the filtered items
+	FilteredItems.Empty();
+	for (auto& Item : PaletteItems)
+	{
+		if (TypeFilter->PassesFilter(Item))
+		{
+			FilteredItems.Add(Item);
+		}
+	}
+
+	// Refresh the appropriate view
+	RefreshActivePaletteViewWidget();
+
+	bIsRebuildTimerRegistered = false;
+	return EActiveTimerReturnType::Stop;
+}
+
+EActiveTimerReturnType SAssetPlacementPalette::RefreshPaletteItems(double InCurrentTime, float InDeltaTime)
+{
+	// Do not refresh the palette if we're waiting on a rebuild
+	if (!bItemsNeedRebuild)
+	{
+		RefreshActivePaletteViewWidget();
+	}
+
+	bIsRefreshTimerRegistered = false;
+	return EActiveTimerReturnType::Stop;
+}
+
+#undef LOCTEXT_NAMESPACE
diff --git a/Engine/Plugins/Experimental/AssetPlacementEdMode/Source/AssetPlacementEdMode/Private/SAssetPlacementPalette.h b/Engine/Plugins/Experimental/AssetPlacementEdMode/Source/AssetPlacementEdMode/Private/SAssetPlacementPalette.h
new file mode 100644
--- /dev/null
+++ b/Engine/Plugins/Experimental/AssetPlacementEdMode/Source/AssetPlacementEdMode/Private/SAssetPlacementPalette.h
@@ -0,0 +1,222 @@
+// Copyright Epic Games, Inc. All Rights Reserved.
+
+#pragma once
+
+#include "CoreMinimal.h"
+#include "Layout/Visibility.h"
+#include "Input/Reply.h"
+#include "Widgets/DeclarativeSyntaxSupport.h"
+#include "Widgets/SWidget.h"
+#include "Widgets/SCompoundWidget.h"
+#include "Framework/Text/SlateHyperlinkRun.h"
+#include "Widgets/Input/SComboButton.h"
+#include "Widgets/Views/SHeaderRow.h"
+#include "Widgets/Views/STableViewBase.h"
+#include "Widgets/Views/STableRow.h"
+#include "Widgets/Views/STileView.h"
+#include "Widgets/Views/STreeView.h"
+#include "Misc/TextFilter.h"
+#include "PlacementPaletteAsset.h"
+
+class FAssetThumbnailPool;
+class FAssetPlacementPaletteItemModel;
+class FMenuBuilder;
+class FUICommandList;
+class IDetailsView;
+struct FAssetData;
+class UAssetPlacementSettings;
+class IPropertyHandle;
+class SSearchBox;
+class UPlacementPaletteClient;
+
+typedef TSharedPtr<FAssetPlacementPaletteItemModel> FPlacementPaletteItemModelPtr;
+typedef STreeView<FPlacementPaletteItemModelPtr> SPlacementTypeTreeView;
+typedef STileView<FPlacementPaletteItemModelPtr> SPlacementTypeTileView;
+
+/** The palette of Placement types available for use by the Placement edit mode */
+class SAssetPlacementPalette : public SCompoundWidget
+{
+public:
+	/** View modes supported by the palette */
+	enum class EViewMode : uint8
+	{
+		Thumbnail,
+		Tree
+	};
+
+	SLATE_BEGIN_ARGS(SAssetPlacementPalette) {}
+		SLATE_ARGUMENT(TWeakPtr<IDetailsView>, ItemDetailsView)
+	SLATE_END_ARGS()
+
+	void Construct(const FArguments& InArgs);
+	~SAssetPlacementPalette();
+
+	// SWidget interface
+	virtual FReply OnKeyDown(const FGeometry& MyGeometry, const FKeyEvent& InKeyEvent) override;
+
+	/** Updates the Placement palette, optionally doing a full rebuild of the items in the palette as well */
+	void UpdatePalette(bool bRebuildItems = false);
+
+	/** Refreshes the Placement palette */
+	void RefreshPalette();
+
+	/** @return True if the given view mode is the active view mode */
+	bool IsActiveViewMode(EViewMode ViewMode) const;
+
+	/** @return True if tooltips should be shown when hovering over Placement type items in the palette */
+	bool ShouldShowTooltips() const;
+
+	/** @return The current search filter text */
+	FText GetSearchText() const;
+
+private:
+	/** Adds the Placement type asset to the instanced Placement actor's list of types. */
+	void AddPlacementType(const FAssetData& AssetData);
+
+	/**
+	 * Removes all items from the palette.
+	 */
+	void ClearPalette();
+	void OnClearPalette();
+	void SetPaletteToAssetDataList(TArrayView<const FAssetData> InAssetDatas);
+
+	/** Refreshes the active palette view widget */
+	void RefreshActivePaletteViewWidget();
+
+	/** Creates the palette views */
+	TSharedRef<class SWidgetSwitcher> CreatePaletteViews();
+
+	/** Adds the displayed name of the Placement type for filtering */
+	void GetPaletteItemFilterString(FPlacementPaletteItemModelPtr PaletteItemModel, TArray<FString>& OutArray) const;
+
+	/** Handles changes to the search filter text */
+	void OnSearchTextChanged(const FText& InFilterText);
+
+	bool ShouldFilterAsset(const FAssetData& InAssetData);
+
+	void OnResetPaletteAssetClicked();
+	void OnSavePaletteAssetClicked();
+	EVisibility GetContentBrowserMirrorVisibility() const;
+	void OnContentBrowserMirrorButtonClicked(ECheckBoxState InState);
+	void OnContentBrowserSelectionChanged(const TArray<FAssetData>& NewSelectedAssets, bool bIsPrimaryBrowser);
+	void SetupContentBrowserMirroring(bool bInMirrorContentBrowser);
+	void OnSetPaletteAsset(const FAssetData& InAssetData);
+	void SetPaletteItems(TArrayView<const TObjectPtr<UPlacementPaletteClient>> InPaletteItems);
+	FString GetPalettePath() const;
+
+	/** Sets the view mode of the palette */
+	void SetViewMode(EViewMode NewViewMode);
+
+	/** Sets whether to show tooltips when hovering over Placement type items in the palette */
+	void ToggleShowTooltips();
+
+	/** Creates the view options menu */
+	TSharedRef<SWidget> GetViewOptionsMenuContent();
+
+	TSharedPtr<SListView<FPlacementPaletteItemModelPtr>> GetActiveViewWidget() const;
+
+	/** Gets the visibility of the drop hint text for when the palette or filtered palette is empty */
+	EVisibility GetDropHintTextVisibility() const;
+
+	bool OnAreAssetsValidForDrop(TArrayView<FAssetData> DraggedAssets) const;
+
+	/** Handles dropping of a mesh or Placement type into the palette */
+	void HandlePlacementDropped(const FDragDropEvent& DragDropEvent, TArrayView<FAssetData> DraggedAssets);
+
+	/** @returns true if there are any items in the palette. */
+	bool HasAnyItemInPalette() const;
+
+private:	// CONTEXT MENU
+
+	/** @return the SWidget containing the context menu */
+	TSharedPtr<SWidget> ConstructPlacementTypeContextMenu();
+
+	/** Called when the selected set of items changes */
+	void OnPaletteSelectionChanged(FPlacementPaletteItemModelPtr Item, ESelectInfo::Type SelectInfo);
+
+	/** Handler for 'Show in CB' command  */
+	void OnShowPlacementTypeInCB();
+
+private:	// THUMBNAIL VIEW
+
+	/** Creates a thumbnail tile for the given Placement type */
+	TSharedRef<ITableRow> GenerateTile(FPlacementPaletteItemModelPtr Item, const TSharedRef<STableViewBase>& OwnerTable);
+
+	/** Gets the scaled thumbnail tile size */
+	float GetScaledThumbnailSize() const;
+
+	/** Gets the current scale of the thumbnail tiles */
+	float GetThumbnailScale() const;
+
+	/** Sets the current scale of the thumbnail tiles */
+	void SetThumbnailScale(float InScale);
+
+	/** Gets whether the thumbnail scaling slider is visible */
+	EVisibility GetThumbnailScaleSliderVisibility() const;
+
+private:	// TREE VIEW
+
+	/** Generates a row widget for Placement mesh item */
+	TSharedRef<ITableRow> TreeViewGenerateRow(FPlacementPaletteItemModelPtr Item, const TSharedRef<STableViewBase>& OwnerTable);
+	void TreeViewGetChildren(FPlacementPaletteItemModelPtr Item, TArray<FPlacementPaletteItemModelPtr>& OutChildren);
+
+	/** Text for Placement meshes list header */
+	FText GetTypeColumnHeaderText() const;
+
+	/** Mesh list sorting support */
+	EColumnSortMode::Type GetMeshColumnSortMode() const;
+	void OnTypeColumnSortModeChanged(EColumnSortPriority::Type InPriority, const FName& InColumnName, EColumnSortMode::Type InSortMode);
+
+private:
+	/** Active timer handler to update the items in the palette */
+	EActiveTimerReturnType UpdatePaletteItems(double InCurrentTime, float InDeltaTime);
+
+	/** Active timer handler to refresh the palette */
+	EActiveTimerReturnType RefreshPaletteItems(double InCurrentTime, float InDeltaTime);
+
+private:
+	typedef TTextFilter<FPlacementPaletteItemModelPtr> PlacementTypeTextFilter;
+	TSharedPtr<PlacementTypeTextFilter> TypeFilter;
+
+	/** All the items in the palette (unfiltered) */
+	TArray<FPlacementPaletteItemModelPtr> PaletteItems;
+
+	/** The filtered list of types to display in the palette */
+	TArray<FPlacementPaletteItemModelPtr> FilteredItems;
+
+	/** Switches between the thumbnail and tree views */
+	TSharedPtr<class SWidgetSwitcher> WidgetSwitcher;
+
+	/** The header row of the Placement mesh tree */
+	TSharedPtr<class SHeaderRow> TreeViewHeaderRow;
+
+	/** Placement type thumbnails widget  */
+	TSharedPtr<SPlacementTypeTileView> TileViewWidget;
+
+	/** Placement type tree widget  */
+	TSharedPtr<SPlacementTypeTreeView> TreeViewWidget;
+
+	/** Placement mesh details widget  */
+	TWeakPtr<IDetailsView> ItemDetailsWidget;
+
+	/** Placement items search box widget */
+	TSharedPtr<SSearchBox> SearchBoxPtr;
+
+	/** Command list for binding functions for the context menu. */
+	TSharedPtr<FUICommandList> UICommandList;
+
+	/** Thumbnail pool for rendering mesh thumbnails */
+	TSharedPtr<FAssetThumbnailPool> ThumbnailPool;
+
+	FSoftObjectPath PalettePath;
+
+	bool bItemsNeedRebuild = true;
+	bool bShowFullTooltips = true;
+	bool bIsRebuildTimerRegistered = true;
+	bool bIsRefreshTimerRegistered = true;
+	bool bIsMirroringContentBrowser = false;
+	EViewMode ActiveViewMode;
+	EColumnSortMode::Type ActiveSortOrder = EColumnSortMode::Type::Ascending;
+
+	float PaletteThumbnailScale = .3f;
+};
diff --git a/Engine/Plugins/Experimental/AssetPlacementEdMode/Source/AssetPlacementEdMode/Private/Tools/PlacementBrushToolBase.cpp b/Engine/Plugins/Experimental/AssetPlacementEdMode/Source/AssetPlacementEdMode/Private/Tools/PlacementBrushToolBase.cpp
new file mode 100644
--- /dev/null
+++ b/Engine/Plugins/Experimental/AssetPlacementEdMode/Source/AssetPlacementEdMode/Private/Tools/PlacementBrushToolBase.cpp
@@ -0,0 +1,332 @@
+// Copyright Epic Games, Inc. All Rights Reserved.
+
+#include "Tools/PlacementBrushToolBase.h"
+#include "UObject/Object.h"
+#include "ToolContextInterfaces.h"
+#include "InteractiveToolManager.h"
+#include "InstancedFoliageActor.h"
+#include "FoliageHelper.h"
+#include "Components/PrimitiveComponent.h"
+#include "Components/BrushComponent.h"
+#include "Components/ModelComponent.h"
+#include "LandscapeHeightfieldCollisionComponent.h"
+#include "FoliageInstancedStaticMeshComponent.h"
+#include "AssetPlacementSettings.h"
+#include "Elements/Framework/EngineElementsLibrary.h"
+#include "Elements/Framework/TypedElementRegistry.h"
+#include "Elements/Framework/TypedElementSelectionSet.h"
+#include "Elements/Interfaces/TypedElementWorldInterface.h"
+#include "Elements/Interfaces/TypedElementSelectionInterface.h"
+#include "ActorPartition/ActorPartitionSubsystem.h"
+#include "Editor.h"
+#include "Modes/PlacementModeSubsystem.h"
+#include "ActorFactories/ActorFactory.h"
+#include "BaseGizmos/GizmoRenderingUtil.h"
+#include "Tools/AssetEditorContextInterface.h"
+#include "ContextObjectStore.h"
+
+bool UPlacementToolBuilderBase::CanBuildTool(const FToolBuilderState& SceneState) const
+{	
+	TWeakObjectPtr<const UAssetPlacementSettings> PlacementSettings = GEditor->GetEditorSubsystem<UPlacementModeSubsystem>()->GetModeSettingsObject();
+	return PlacementSettings.IsValid() && PlacementSettings->GetActivePaletteItems().Num();
+}
+
+UInteractiveTool* UPlacementToolBuilderBase::BuildTool(const FToolBuilderState& SceneState) const
+{
+	return FactoryToolInstance(SceneState.ToolManager);
+}
+
+bool UPlacementBrushToolBase::HitTest(const FRay& Ray, FHitResult& OutHit)
+{
+	const FVector TraceStart(Ray.Origin);
+	const FVector TraceEnd(Ray.Origin + Ray.Direction * HALF_WORLD_MAX);
+
+	return FindHitResultWithStartAndEndTraceVectors(OutHit, TraceStart, TraceEnd);
+}
+
+bool UPlacementBrushToolBase::AreAllTargetsValid() const
+{
+	return Target ? Target->IsValid() : true;
+}
+
+void UPlacementBrushToolBase::Render(IToolsContextRenderAPI* RenderAPI)
+{
+	Super::Render(RenderAPI);
+
+	// Transform the brush radius to standard pixel size
+	LastBrushStampWorldToPixelScale = GizmoRenderingUtil::CalculateLocalPixelToWorldScale(RenderAPI->GetSceneView(), LastBrushStamp.WorldPosition);
+}
+
+void UPlacementBrushToolBase::OnClickPress(const FInputDeviceRay& PressPos)
+{
+	LastDeviceInputRay = PressPos;
+	Super::OnClickPress(PressPos);
+}
+
+void UPlacementBrushToolBase::OnClickDrag(const FInputDeviceRay& DragPos)
+{
+	LastDeviceInputRay = DragPos;
+	GetToolManager()->PostInvalidation();
+	Super::OnClickDrag(DragPos);
+}
+
+double UPlacementBrushToolBase::EstimateMaximumTargetDimension()
+{
+	return 1000.0;
+}
+
+bool UPlacementBrushToolBase::FindHitResultWithStartAndEndTraceVectors(FHitResult& OutHit, const FVector& TraceStart, const FVector& TraceEnd, float TraceRadius)
+{
+	UWorld* EditingWorld = GetToolManager()->GetWorld();
+	constexpr TCHAR NAME_PlacementBrushTool[] = TEXT("PlacementBrushTool");
+
+	TWeakObjectPtr<const UAssetPlacementSettings> PlacementSettings = GEditor->GetEditorSubsystem<UPlacementModeSubsystem>()->GetModeSettingsObject();
+
+	auto FilterFunc = [PlacementSettings](const UPrimitiveComponent* InComponent) {
+		if (InComponent && PlacementSettings.IsValid())
+		{
+			bool bFoliageOwned = InComponent->GetOwner() && FFoliageHelper::IsOwnedByFoliage(InComponent->GetOwner());
+			const bool bAllowLandscape = PlacementSettings->bLandscape;
+			const bool bAllowStaticMesh = PlacementSettings->bStaticMeshes;
+			const bool bAllowBSP = PlacementSettings->bBSP;
+			const bool bAllowFoliage = PlacementSettings->bFoliage;
+			const bool bAllowTranslucent = PlacementSettings->bTranslucent;
+
+			// allow list
+			bool bAllowed =
+				(bAllowLandscape && InComponent->IsA(ULandscapeHeightfieldCollisionComponent::StaticClass())) ||
+				(bAllowStaticMesh && InComponent->IsA(UStaticMeshComponent::StaticClass()) && !InComponent->IsA(UFoliageInstancedStaticMeshComponent::StaticClass()) && !bFoliageOwned) ||
+				(bAllowBSP && (InComponent->IsA(UBrushComponent::StaticClass()) || InComponent->IsA(UModelComponent::StaticClass()))) ||
+				(bAllowFoliage && (InComponent->IsA(UFoliageInstancedStaticMeshComponent::StaticClass()) || bFoliageOwned));
+
+			// deny list
+			bAllowed &=
+				(bAllowTranslucent || !(InComponent->GetMaterial(0) && IsTranslucentBlendMode(InComponent->GetMaterial(0)->GetBlendMode())));
+
+			return bAllowed;
+		}
+
+		return false; };
+
+	return AInstancedFoliageActor::FoliageTrace(EditingWorld, OutHit, FDesiredFoliageInstance(TraceStart, TraceEnd, /* FoliageType= */ nullptr, TraceRadius), NAME_PlacementBrushTool, /* bReturnFaceIndex */ false, FilterFunc);
+}
+
+FTransform UPlacementBrushToolBase::GenerateTransformFromHitLocationAndNormal(const FVector& InLocation, const FVector& InNormal)
+{
+	const UAssetPlacementSettings* PlacementSettings = GEditor->GetEditorSubsystem<UPlacementModeSubsystem>()->GetModeSettingsObject();
+	FTransform FinalizedTransform(GenerateRandomRotation(PlacementSettings), InLocation, GenerateRandomScale(PlacementSettings));
+	return FinalizeTransform(FinalizedTransform, InNormal, PlacementSettings);
+}
+
+FQuat UPlacementBrushToolBase::GenerateRandomRotation(const UAssetPlacementSettings* PlacementSettings)
+{
+	if (!PlacementSettings)
+	{
+		return FQuat::Identity;
+	}
+
+	FRotator GeneratedRotation = FRotator::ZeroRotator;
+	auto GetRandomSignedValueInRange = [](const FFloatInterval& Range, bool bAllowSigned) -> float
+	{
+		float Sign = (FMath::RandBool() && bAllowSigned) ? -1.0f : 1.0f;
+		return Range.Interpolate(FMath::FRand()) * Sign;
+	};
+
+	if (PlacementSettings->bUseRandomRotationX)
+	{
+		GeneratedRotation.Roll = GetRandomSignedValueInRange(PlacementSettings->RandomRotationX, PlacementSettings->bAllowNegativeRotationX);
+	}
+
+	if (PlacementSettings->bUseRandomRotationY)
+	{
+		GeneratedRotation.Pitch = GetRandomSignedValueInRange(PlacementSettings->RandomRotationY, PlacementSettings->bAllowNegativeRotationY);
+	}
+
+	if (PlacementSettings->bUseRandomRotationZ)
+	{
+		GeneratedRotation.Yaw = GetRandomSignedValueInRange(PlacementSettings->RandomRotationZ, PlacementSettings->bAllowNegativeRotationZ);
+	}
+
+	return GeneratedRotation.Quaternion();
+}
+
+FVector UPlacementBrushToolBase::GenerateRandomScale(const UAssetPlacementSettings* PlacementSettings)
+{
+	FVector GeneratedScale(1.0f);
+	if (!PlacementSettings || !PlacementSettings->bUseRandomScale)
+	{
+		return GeneratedScale;
+	}
+
+	auto GenerateRandomScaleComponent = [PlacementSettings]() -> float
+	{
+		float Sign = (FMath::RandBool() && PlacementSettings->bAllowNegativeScale) ? -1.0f : 1.0f;
+		return PlacementSettings->ScaleRange.Interpolate(FMath::FRand()) * Sign;
+	};
+
+	switch (PlacementSettings->ScalingType)
+	{
+		case EFoliageScaling::Free:
+		{
+			GeneratedScale = FVector(GenerateRandomScaleComponent(), GenerateRandomScaleComponent(), GenerateRandomScaleComponent());
+			break;
+		}
+		case EFoliageScaling::Uniform:
+		{
+			float ScaleComponent = GenerateRandomScaleComponent();
+			GeneratedScale = FVector(ScaleComponent, ScaleComponent, ScaleComponent);
+			break;
+		}
+		case EFoliageScaling::LockXY:
+		{
+			GeneratedScale.Z = GenerateRandomScaleComponent();
+			break;
+		}
+		case EFoliageScaling::LockYZ:
+		{
+			GeneratedScale.X = GenerateRandomScaleComponent();
+			break;
+		}
+		case EFoliageScaling::LockXZ:
+		{
+			GeneratedScale.Y = GenerateRandomScaleComponent();
+			break;
+		}
+	}
+
+	return GeneratedScale;
+}
+
+FQuat UPlacementBrushToolBase::AlignRotationWithNormal(const FQuat& InRotation, const FVector& InNormal, EAxis::Type InAlignmentAxis, bool bInvertAxis)
+{
+	FVector AlignmentVector = FVector::UpVector;
+	switch (InAlignmentAxis)
+	{
+		case EAxis::Type::X:
+		{
+			AlignmentVector = bInvertAxis ? FVector::BackwardVector : FVector::ForwardVector;
+		}
+		break;
+
+		case EAxis::Type::Y:
+		{
+			AlignmentVector = bInvertAxis ? FVector::LeftVector : FVector::RightVector;
+		}
+		break;
+
+		case EAxis::Type::Z:
+		{
+			AlignmentVector = bInvertAxis ? FVector::DownVector : FVector::UpVector;
+		}
+		break;
+	}
+
+	return FindActorAlignmentRotation(InRotation, AlignmentVector, InNormal);
+}
+
+FTransform UPlacementBrushToolBase::FinalizeTransform(const FTransform& OriginalTransform, const FVector& InNormal, const UAssetPlacementSettings* PlacementSettings)
+{
+	if (!PlacementSettings)
+	{
+		return OriginalTransform;
+	}
+
+	FTransform FinalizedTransform(FQuat::Identity, OriginalTransform.GetTranslation(), OriginalTransform.GetScale3D());
+
+	// Add the world offset.
+	FVector WorldOffset = PlacementSettings->WorldLocationOffset;
+	if (PlacementSettings->bScaleWorldLocationOffset)
+	{
+		WorldOffset *= OriginalTransform.GetScale3D();
+	}
+	FinalizedTransform.AddToTranslation(WorldOffset);
+
+	// Align to normal.
+	FQuat AdjustedRotation(OriginalTransform.GetRotation());
+	if (PlacementSettings->bAlignToNormal)
+	{
+		AdjustedRotation = AlignRotationWithNormal(AdjustedRotation, InNormal, PlacementSettings->AxisToAlignWithNormal, PlacementSettings->bInvertNormalAxis);
+	}
+	AdjustedRotation.Normalize();
+	FinalizedTransform.SetRotation(AdjustedRotation);
+
+	// Add the relative offset.
+	{
+		FVector RelativeOffset = PlacementSettings->RelativeLocationOffset;
+		if (PlacementSettings->bScaleRelativeLocationOffset)
+		{
+			RelativeOffset *= OriginalTransform.GetScale3D();
+		}
+		FinalizedTransform.SetTranslation(FinalizedTransform.TransformPosition(RelativeOffset));
+	}
+
+	return FinalizedTransform;
+}
+
+FTypedElementListRef UPlacementBrushToolBase::GetElementsInBrushRadius(const FInputDeviceRay& DragPos) const
+{
+	FTypedElementListRef ElementHandles = UTypedElementRegistry::GetInstance()->CreateElementList();
+
+	const UTypedElementSelectionSet* SelectionSet = nullptr;
+	if (IAssetEditorContextInterface* AssetEditorAPI = GetToolManager()->GetContextObjectStore()->FindContext<IAssetEditorContextInterface>())
+	{
+		SelectionSet = AssetEditorAPI->GetSelectionSet();
+	}
+
+	if (!SelectionSet)
+	{
+		return ElementHandles;
+	}
+
+	UPlacementModeSubsystem* PlacementModeSubsystem = GEditor->GetEditorSubsystem<UPlacementModeSubsystem>();
+	if (!PlacementModeSubsystem || !PlacementModeSubsystem->GetModeSettingsObject())
+	{
+		return ElementHandles;
+	}
+
+	FCollisionQueryParams QueryParams(TEXT("PlacementBrushTool"), SCENE_QUERY_STAT_ONLY(EdMode_PlacementTrace), true);
+	QueryParams.bReturnFaceIndex = false;
+	TArray<FHitResult> Hits;
+	FCollisionShape BrushSphere;
+	BrushSphere.SetSphere(LastBrushStamp.Radius);
+
+	const FVector TraceStart(LastWorldRay.Origin);
+	const FVector TraceEnd(LastWorldRay.Origin + LastWorldRay.Direction * HALF_WORLD_MAX);
+
+	GetToolManager()->GetWorld()->SweepMultiByObjectType(Hits, TraceStart, TraceEnd, FQuat::Identity,
+		FCollisionObjectQueryParams(FCollisionObjectQueryParams::InitType::AllObjects), BrushSphere, QueryParams);
+
+	for (const FHitResult& Hit : Hits)
+	{
+		if (UInstancedStaticMeshComponent* ISMComponent = Cast<UInstancedStaticMeshComponent>(Hit.GetComponent()))
+		{
+			FTypedElementHandle StaticMeshObjectHandle = UEngineElementsLibrary::AcquireEditorObjectElementHandle(ISMComponent->GetStaticMesh());
+			if (!PlacementModeSubsystem->GetModeSettingsObject()->DoesActivePaletteSupportElement(StaticMeshObjectHandle))
+			{
+				continue;
+			}
+
+			TArray<int32> OverlappingInstances = ISMComponent->GetInstancesOverlappingSphere(LastBrushStamp.WorldPosition, LastBrushStamp.Radius);
+			ElementHandles->Reserve(OverlappingInstances.Num());
+			for (const int32& InstanceIdx : OverlappingInstances)
+			{
+				if (FTypedElementHandle SMInstanceHandle = UEngineElementsLibrary::AcquireEditorSMInstanceElementHandle(ISMComponent, InstanceIdx))
+				{
+					ElementHandles->Add(MoveTemp(SMInstanceHandle));
+				}
+			}
+		}
+		else
+		{
+			FTypedElementHandle PrimarySelectionHandle = SelectionSet->GetSelectionElement(UEngineElementsLibrary::AcquireEditorComponentElementHandle(Hit.GetComponent()), ETypedElementSelectionMethod::Primary);
+			if (!PlacementModeSubsystem->GetModeSettingsObject()->DoesActivePaletteSupportElement(PrimarySelectionHandle))
+			{
+				continue;
+			}
+
+			ElementHandles->Add(MoveTemp(PrimarySelectionHandle));
+		}
+	}
+
+	return ElementHandles;
+}
diff --git a/Engine/Plugins/Experimental/AssetPlacementEdMode/Source/AssetPlacementEdMode/Private/Tools/PlacementBrushToolBase.h b/Engine/Plugins/Experimental/AssetPlacementEdMode/Source/AssetPlacementEdMode/Private/Tools/PlacementBrushToolBase.h
new file mode 100644
--- /dev/null
+++ b/Engine/Plugins/Experimental/AssetPlacementEdMode/Source/AssetPlacementEdMode/Private/Tools/PlacementBrushToolBase.h
@@ -0,0 +1,59 @@
+// Copyright Epic Games, Inc. All Rights Reserved.
+
+#pragma once
+
+#include "InputState.h"
+#include "BaseTools/BaseBrushTool.h"
+#include "Elements/Framework/TypedElementListFwd.h"
+
+#include "PlacementBrushToolBase.generated.h"
+
+struct FHitResult;
+class UAssetPlacementSettings;
+
+UCLASS(Abstract, MinimalAPI)
+class UPlacementToolBuilderBase : public UInteractiveToolBuilder
+{
+	GENERATED_BODY()
+
+public:
+	virtual bool CanBuildTool(const FToolBuilderState& SceneState) const override;
+	virtual UInteractiveTool* BuildTool(const FToolBuilderState& SceneState) const override;
+
+protected:
+	virtual UPlacementBrushToolBase* FactoryToolInstance(UObject* Outer) const PURE_VIRTUAL(UPlacementToolBuilderBase::FactoryToolInstance, return nullptr; );
+};
+
+UCLASS(Abstract, MinimalAPI)
+class UPlacementBrushToolBase : public UBaseBrushTool
+{
+	GENERATED_BODY()
+
+	friend class UPlacementToolBuilderBase;
+	
+public:
+	virtual bool HitTest(const FRay& Ray, FHitResult& OutHit) override;
+	virtual bool AreAllTargetsValid() const override;
+	virtual void Render(IToolsContextRenderAPI* RenderAPI) override;
+	virtual void OnClickPress(const FInputDeviceRay& PressPos) override;
+	virtual void OnClickDrag(const FInputDeviceRay& DragPos) override;
+
+protected:
+	virtual double EstimateMaximumTargetDimension() override;
+	bool FindHitResultWithStartAndEndTraceVectors(FHitResult& OutHit, const FVector& TraceStart, const FVector& TraceEnd, float TraceRadius = 0.0f);
+	static FTransform GenerateTransformFromHitLocationAndNormal(const FVector& InLocation, const FVector& InNormal);
+
+	// Gets a random rotation, and aligns it, based on the placement settings.
+	static FQuat GenerateRandomRotation(const UAssetPlacementSettings* PlacementSettings);
+
+	// Generates a random scale, based on placement settings.
+	static FVector GenerateRandomScale(const UAssetPlacementSettings* PlacementSettings);
+
+	// Updates the last generated rotation to realign with the current brush position and normal.
+	static FQuat AlignRotationWithNormal(const FQuat& InRotation, const FVector& InNormal, EAxis::Type InAlignmentAxis, bool bInvertAxis);
+	static FTransform FinalizeTransform(const FTransform& OriginalTransform, const FVector& InNormal, const UAssetPlacementSettings* PlacementSettings);
+	FTypedElementListRef GetElementsInBrushRadius(const FInputDeviceRay& DragPos) const;
+
+	float LastBrushStampWorldToPixelScale;
+	FInputDeviceRay LastDeviceInputRay = FInputDeviceRay(FRay());
+};
diff --git a/Engine/Plugins/Experimental/AssetPlacementEdMode/Source/AssetPlacementEdMode/Private/Tools/PlacementClickDragToolBase.cpp b/Engine/Plugins/Experimental/AssetPlacementEdMode/Source/AssetPlacementEdMode/Private/Tools/PlacementClickDragToolBase.cpp
new file mode 100644
--- /dev/null
+++ b/Engine/Plugins/Experimental/AssetPlacementEdMode/Source/AssetPlacementEdMode/Private/Tools/PlacementClickDragToolBase.cpp
@@ -0,0 +1,28 @@
+// Copyright Epic Games, Inc. All Rights Reserved.
+
+#include "Tools/PlacementClickDragToolBase.h"
+#include "UObject/Object.h"
+#include "ToolContextInterfaces.h"
+#include "InteractiveToolManager.h"
+#include "BaseGizmos/BrushStampIndicator.h"
+
+void UPlacementClickDragToolBase::Setup()
+{
+	Super::Setup();
+
+	RemoveToolPropertySource(BrushProperties);
+}
+
+double UPlacementClickDragToolBase::EstimateMaximumTargetDimension()
+{
+	return Super::EstimateMaximumTargetDimension();
+}
+
+void UPlacementClickDragToolBase::SetupBrushStampIndicator()
+{
+	Super::SetupBrushStampIndicator();
+
+	BrushStampIndicator->bDrawRadiusCircle = false;
+	BrushProperties->BrushSize = .5f;
+	BrushProperties->BrushFalloffAmount = 1.f;
+}
diff --git a/Engine/Plugins/Experimental/AssetPlacementEdMode/Source/AssetPlacementEdMode/Private/Tools/PlacementClickDragToolBase.h b/Engine/Plugins/Experimental/AssetPlacementEdMode/Source/AssetPlacementEdMode/Private/Tools/PlacementClickDragToolBase.h
new file mode 100644
--- /dev/null
+++ b/Engine/Plugins/Experimental/AssetPlacementEdMode/Source/AssetPlacementEdMode/Private/Tools/PlacementClickDragToolBase.h
@@ -0,0 +1,19 @@
+// Copyright Epic Games, Inc. All Rights Reserved.
+
+#pragma once
+
+#include "Tools/PlacementBrushToolBase.h"
+
+#include "PlacementClickDragToolBase.generated.h"
+
+UCLASS(Abstract, MinimalAPI)
+class UPlacementClickDragToolBase : public UPlacementBrushToolBase
+{
+	GENERATED_BODY()
+public:
+	virtual void Setup() override;
+
+protected:
+	virtual double EstimateMaximumTargetDimension() override;
+	virtual void SetupBrushStampIndicator() override;
+};
diff --git a/Engine/Plugins/Experimental/AssetPlacementEdMode/Source/AssetPlacementEdMode/Private/Tools/PlacementEraseTool.cpp b/Engine/Plugins/Experimental/AssetPlacementEdMode/Source/AssetPlacementEdMode/Private/Tools/PlacementEraseTool.cpp
new file mode 100644
--- /dev/null
+++ b/Engine/Plugins/Experimental/AssetPlacementEdMode/Source/AssetPlacementEdMode/Private/Tools/PlacementEraseTool.cpp
@@ -0,0 +1,74 @@
+// Copyright Epic Games, Inc. All Rights Reserved.
+
+#include "Tools/PlacementEraseTool.h"
+#include "UObject/Object.h"
+#include "ToolContextInterfaces.h"
+#include "InteractiveToolManager.h"
+#include "Elements/Framework/EngineElementsLibrary.h"
+#include "Elements/Framework/TypedElementCommonActions.h"
+#include "Elements/Framework/TypedElementHandle.h"
+#include "Elements/Framework/TypedElementRegistry.h"
+#include "Elements/Interfaces/TypedElementObjectInterface.h"
+#include "Tools/AssetEditorContextInterface.h"
+#include "Editor.h"
+#include "AssetPlacementEdMode.h"
+#include "AssetPlacementSettings.h"
+#include "Modes/PlacementModeSubsystem.h"
+#include "Toolkits/IToolkitHost.h"
+#include "UObject/GCObjectScopeGuard.h"
+#include "EditorModeManager.h"
+#include "ContextObjectStore.h"
+
+constexpr TCHAR UPlacementModeEraseTool::ToolName[];
+
+UPlacementBrushToolBase* UPlacementModeEraseToolBuilder::FactoryToolInstance(UObject* Outer) const
+{
+	return NewObject<UPlacementModeEraseTool>(Outer);
+}
+
+void UPlacementModeEraseTool::OnBeginDrag(const FRay& Ray)
+{
+	Super::OnBeginDrag(Ray);
+
+	GetToolManager()->BeginUndoTransaction(NSLOCTEXT("AssetPlacementEdMode", "BrushErase", "Erase Painted Elements"));
+}
+
+void UPlacementModeEraseTool::OnEndDrag(const FRay& Ray)
+{
+	GetToolManager()->EndUndoTransaction();
+
+	Super::OnEndDrag(Ray);
+}
+
+void UPlacementModeEraseTool::OnTick(float DeltaTime)
+{
+	if (!bInBrushStroke)
+	{
+		return;
+	}
+
+	IAssetEditorContextInterface* AssetEditorContext = GetToolManager()->GetContextObjectStore()->FindContext<IAssetEditorContextInterface>();
+	if (!AssetEditorContext)
+	{
+		return;
+	}
+
+	UTypedElementCommonActions* ElementCommonActions = AssetEditorContext->GetCommonActions();
+	if (!ElementCommonActions)
+	{
+		return;
+	}
+
+	UTypedElementSelectionSet* SelectionSet = AssetEditorContext->GetMutableSelectionSet();
+	if (!SelectionSet)
+	{
+		return;
+	}
+
+	FTypedElementListRef ElementsToDelete = GetElementsInBrushRadius(LastDeviceInputRay);
+	if (ElementsToDelete->HasElements())
+	{
+		FTypedElementListRef NormalizedElementsToDelete = SelectionSet->GetNormalizedElementList(ElementsToDelete, FTypedElementSelectionNormalizationOptions());
+		ElementCommonActions->DeleteNormalizedElements(NormalizedElementsToDelete, AssetEditorContext->GetEditingWorld(), SelectionSet, FTypedElementDeletionOptions());
+	}
+}
diff --git a/Engine/Plugins/Experimental/AssetPlacementEdMode/Source/AssetPlacementEdMode/Private/Tools/PlacementEraseTool.h b/Engine/Plugins/Experimental/AssetPlacementEdMode/Source/AssetPlacementEdMode/Private/Tools/PlacementEraseTool.h
new file mode 100644
--- /dev/null
+++ b/Engine/Plugins/Experimental/AssetPlacementEdMode/Source/AssetPlacementEdMode/Private/Tools/PlacementEraseTool.h
@@ -0,0 +1,31 @@
+// Copyright Epic Games, Inc. All Rights Reserved.
+
+#pragma once
+
+#include "InteractiveToolBuilder.h"
+#include "Tools/PlacementBrushToolBase.h"
+
+#include "PlacementEraseTool.generated.h"
+
+UCLASS(Transient, MinimalAPI)
+class UPlacementModeEraseToolBuilder : public UPlacementToolBuilderBase
+{
+	GENERATED_BODY()
+
+protected:
+	virtual UPlacementBrushToolBase* FactoryToolInstance(UObject* Outer) const override;
+};
+
+UCLASS(MinimalAPI)
+class UPlacementModeEraseTool : public UPlacementBrushToolBase
+{
+	GENERATED_BODY()
+public:
+	constexpr static TCHAR ToolName[] = TEXT("EraseTool");
+
+	virtual void OnBeginDrag(const FRay& Ray) override;
+	virtual void OnEndDrag(const FRay& Ray) override;
+
+protected:
+	virtual void OnTick(float DeltaTime) override;
+};
diff --git a/Engine/Plugins/Experimental/AssetPlacementEdMode/Source/AssetPlacementEdMode/Private/Tools/PlacementFillTool.cpp b/Engine/Plugins/Experimental/AssetPlacementEdMode/Source/AssetPlacementEdMode/Private/Tools/PlacementFillTool.cpp
new file mode 100644
--- /dev/null
+++ b/Engine/Plugins/Experimental/AssetPlacementEdMode/Source/AssetPlacementEdMode/Private/Tools/PlacementFillTool.cpp
@@ -0,0 +1,18 @@
+// Copyright Epic Games, Inc. All Rights Reserved.
+
+#include "Tools/PlacementFillTool.h"
+#include "UObject/Object.h"
+#include "ToolContextInterfaces.h"
+#include "InteractiveToolManager.h"
+
+constexpr TCHAR UPlacementModeFillTool::ToolName[];
+
+bool UPlacementModeFillToolBuilder::CanBuildTool(const FToolBuilderState& SceneState) const
+{
+	return true;
+}
+
+UInteractiveTool* UPlacementModeFillToolBuilder::BuildTool(const FToolBuilderState& SceneState) const
+{
+	return NewObject<UPlacementModeFillTool>(SceneState.ToolManager, UPlacementModeFillTool::ToolName);
+}
diff --git a/Engine/Plugins/Experimental/AssetPlacementEdMode/Source/AssetPlacementEdMode/Private/Tools/PlacementFillTool.h b/Engine/Plugins/Experimental/AssetPlacementEdMode/Source/AssetPlacementEdMode/Private/Tools/PlacementFillTool.h
new file mode 100644
--- /dev/null
+++ b/Engine/Plugins/Experimental/AssetPlacementEdMode/Source/AssetPlacementEdMode/Private/Tools/PlacementFillTool.h
@@ -0,0 +1,26 @@
+// Copyright Epic Games, Inc. All Rights Reserved.
+
+#pragma once
+
+#include "InteractiveToolBuilder.h"
+#include "Tools/PlacementClickDragToolBase.h"
+
+#include "PlacementFillTool.generated.h"
+
+UCLASS(Transient, MinimalAPI)
+class UPlacementModeFillToolBuilder : public UInteractiveToolBuilder
+{
+	GENERATED_BODY()
+
+public:
+	virtual bool CanBuildTool(const FToolBuilderState& SceneState) const override;
+	virtual UInteractiveTool* BuildTool(const FToolBuilderState& SceneState) const override;
+};
+
+UCLASS(MinimalAPI)
+class UPlacementModeFillTool : public UPlacementClickDragToolBase
+{
+	GENERATED_BODY()
+public:
+	constexpr static TCHAR ToolName[] = TEXT("FillTool");
+};
diff --git a/Engine/Plugins/Experimental/AssetPlacementEdMode/Source/AssetPlacementEdMode/Private/Tools/PlacementLassoSelectTool.cpp b/Engine/Plugins/Experimental/AssetPlacementEdMode/Source/AssetPlacementEdMode/Private/Tools/PlacementLassoSelectTool.cpp
new file mode 100644
--- /dev/null
+++ b/Engine/Plugins/Experimental/AssetPlacementEdMode/Source/AssetPlacementEdMode/Private/Tools/PlacementLassoSelectTool.cpp
@@ -0,0 +1,70 @@
+// Copyright Epic Games, Inc. All Rights Reserved.
+
+#include "Tools/PlacementLassoSelectTool.h"
+
+#include "AssetPlacementEdMode.h"
+#include "AssetPlacementSettings.h"
+#include "Editor.h"
+#include "Elements/Framework/EngineElementsLibrary.h"
+#include "Elements/Framework/TypedElementHandle.h"
+#include "Elements/Framework/TypedElementRegistry.h"
+#include "Elements/Interfaces/TypedElementObjectInterface.h"
+#include "InteractiveToolManager.h"
+#include "ToolContextInterfaces.h"
+#include "UObject/Object.h"
+#include "BaseBehaviors/KeyAsModifierInputBehavior.h"
+#include "Elements/Framework/TypedElementSelectionSet.h"
+#include "Modes/PlacementModeSubsystem.h"
+#include "Tools/AssetEditorContextInterface.h"
+#include "EditorModeManager.h"
+#include "ContextObjectStore.h"
+
+constexpr TCHAR UPlacementModeLassoSelectTool::ToolName[];
+
+namespace PlacementModeLassoToolInternal
+{
+	FTypedElementSelectionOptions SelectionOptions {};
+}
+
+UPlacementBrushToolBase* UPlacementModeLassoSelectToolBuilder::FactoryToolInstance(UObject* Outer) const
+{
+	return NewObject<UPlacementModeLassoSelectTool>(Outer);
+}
+
+void UPlacementModeLassoSelectTool::OnBeginDrag(const FRay& Ray)
+{
+	Super::OnBeginDrag(Ray);
+	GetToolManager()->BeginUndoTransaction(NSLOCTEXT("AssetPlacementEdMode", "BrushSelect", "Select Elements"));
+}
+
+void UPlacementModeLassoSelectTool::OnEndDrag(const FRay& Ray)
+{
+	GetToolManager()->EndUndoTransaction();
+	Super::OnEndDrag(Ray);
+}
+
+void UPlacementModeLassoSelectTool::OnTick(float DeltaTime)
+{
+	if (!bInBrushStroke)
+	{
+		return;
+	}
+
+	if (IAssetEditorContextInterface* AssetEditorContext = GetToolManager()->GetContextObjectStore()->FindContext<IAssetEditorContextInterface>())
+	{
+		UTypedElementSelectionSet* SelectionSet = AssetEditorContext->GetMutableSelectionSet();
+		if (SelectionSet)
+		{
+			FTypedElementListRef HitElements = GetElementsInBrushRadius(LastDeviceInputRay);
+			const bool bSelectElements = !bCtrlToggle;
+			if (bSelectElements)
+			{
+				SelectionSet->SelectElements(HitElements, PlacementModeLassoToolInternal::SelectionOptions);
+			}
+			else
+			{
+				SelectionSet->DeselectElements(HitElements, PlacementModeLassoToolInternal::SelectionOptions);
+			}
+		}
+	}
+}
diff --git a/Engine/Plugins/Experimental/AssetPlacementEdMode/Source/AssetPlacementEdMode/Private/Tools/PlacementLassoSelectTool.h b/Engine/Plugins/Experimental/AssetPlacementEdMode/Source/AssetPlacementEdMode/Private/Tools/PlacementLassoSelectTool.h
new file mode 100644
--- /dev/null
+++ b/Engine/Plugins/Experimental/AssetPlacementEdMode/Source/AssetPlacementEdMode/Private/Tools/PlacementLassoSelectTool.h
@@ -0,0 +1,30 @@
+// Copyright Epic Games, Inc. All Rights Reserved.
+
+#pragma once
+
+#include "InteractiveToolBuilder.h"
+#include "Tools/PlacementBrushToolBase.h"
+
+#include "PlacementLassoSelectTool.generated.h"
+
+UCLASS(Transient, MinimalAPI)
+class UPlacementModeLassoSelectToolBuilder : public UPlacementToolBuilderBase
+{
+	GENERATED_BODY()
+
+public:
+	virtual UPlacementBrushToolBase* FactoryToolInstance(UObject* Outer) const override;
+};
+
+UCLASS(MinimalAPI)
+class UPlacementModeLassoSelectTool : public UPlacementBrushToolBase
+{
+	GENERATED_BODY()
+
+public:
+	virtual void OnBeginDrag(const FRay& Ray) override;
+	virtual void OnEndDrag(const FRay& Ray) override;
+	virtual void OnTick(float DeltaTime) override;
+
+	constexpr static TCHAR ToolName[] = TEXT("LassoSelectTool");
+};
diff --git a/Engine/Plugins/Experimental/AssetPlacementEdMode/Source/AssetPlacementEdMode/Private/Tools/PlacementPlaceSingleTool.cpp b/Engine/Plugins/Experimental/AssetPlacementEdMode/Source/AssetPlacementEdMode/Private/Tools/PlacementPlaceSingleTool.cpp
new file mode 100644
--- /dev/null
+++ b/Engine/Plugins/Experimental/AssetPlacementEdMode/Source/AssetPlacementEdMode/Private/Tools/PlacementPlaceSingleTool.cpp
@@ -0,0 +1,413 @@
+// Copyright Epic Games, Inc. All Rights Reserved.
+
+#include "Tools/PlacementPlaceSingleTool.h"
+
+#include "AssetPlacementSettings.h"
+#include "Editor.h"
+#include "InteractiveToolManager.h"
+#include "ScopedTransaction.h"
+#include "ToolDataVisualizer.h"
+#include "BaseBehaviors/ClickDragBehavior.h"
+#include "BaseGizmos/BrushStampIndicator.h"
+#include "BaseGizmos/GizmoRenderingUtil.h"
+#include "Editor/EditorEngine.h"
+#include "Elements/Framework/TypedElementHandle.h"
+#include "Elements/Framework/TypedElementRegistry.h"
+#include "Elements/Framework/TypedElementSelectionSet.h"
+#include "Elements/Interfaces/TypedElementWorldInterface.h"
+#include "Elements/SMInstance/SMInstanceElementData.h" // For SMInstanceElementDataUtil::SMInstanceElementsEnabled
+#include "Modes/PlacementModeSubsystem.h"
+#include "Subsystems/PlacementSubsystem.h"
+#include "UObject/Object.h"
+#include "Tools/AssetEditorContextInterface.h"
+#include "EditorModeManager.h"
+#include "Toolkits/IToolkitHost.h"
+#include "ContextObjectStore.h"
+#include "PlacementPaletteItem.h"
+
+constexpr TCHAR UPlacementModePlaceSingleTool::ToolName[];
+
+UPlacementBrushToolBase* UPlacementModePlaceSingleToolBuilder::FactoryToolInstance(UObject* Outer) const
+{
+	return NewObject<UPlacementModePlaceSingleTool>(Outer);
+}
+
+UPlacementModePlaceSingleTool::UPlacementModePlaceSingleTool() = default;
+UPlacementModePlaceSingleTool::~UPlacementModePlaceSingleTool() = default;
+UPlacementModePlaceSingleTool::UPlacementModePlaceSingleTool(FVTableHelper& Helper)
+	: Super(Helper)
+{
+}
+
+void UPlacementModePlaceSingleTool::Setup()
+{
+	Super::Setup();
+
+	SetupRightClickMouseBehavior();
+
+	bIsTweaking = false;
+
+	SinglePlaceSettings = NewObject<UPlacementModePlaceSingleToolSettings>(this);
+	SinglePlaceSettings->LoadConfig();
+	AddToolPropertySource(SinglePlaceSettings);
+}
+
+void UPlacementModePlaceSingleTool::Shutdown(EToolShutdownType ShutdownType)
+{
+	DestroyPreviewElements();
+
+	// Preserve the selection on exiting the tool, so that a users' state persists as they use the mode.
+	constexpr bool bClearSelectionSet = false;
+	ExitTweakState(bClearSelectionSet);
+
+	Super::Shutdown(ShutdownType);
+
+	SinglePlaceSettings->SaveConfig();
+	SinglePlaceSettings = nullptr;
+}
+
+void UPlacementModePlaceSingleTool::OnClickPress(const FInputDeviceRay& PressPos)
+{
+	Super::OnClickPress(PressPos);
+
+	DestroyPreviewElements();
+	LastBrushStamp.Radius = 100.0f * LastBrushStampWorldToPixelScale;
+
+	UPlacementModeSubsystem* PlacementModeSubsystem = GEditor->GetEditorSubsystem<UPlacementModeSubsystem>();
+	const UAssetPlacementSettings* PlacementSettings = PlacementModeSubsystem ? PlacementModeSubsystem->GetModeSettingsObject() : nullptr;
+	UPlacementSubsystem* PlacementSubsystem = GEditor->GetEditorSubsystem<UPlacementSubsystem>();
+	if (PlacementSubsystem && PlacementSettings)
+	{
+		// Place the Preview data if we managed to get to a valid handled click.
+		FPlacementOptions PlacementOptions;
+		PlacementOptions.bPreferBatchPlacement = true;
+		PlacementOptions.InstancedPlacementGridGuid = PlacementSettings->GetActivePaletteGuid();
+
+		FAssetPlacementInfo FinalizedPlacementInfo = *PlacementInfo;
+		FinalizedPlacementInfo.FinalizedTransform = FinalizeTransform(
+			FTransform(PlacementInfo->FinalizedTransform.GetRotation(), LastBrushStamp.WorldPosition, PlacementInfo->FinalizedTransform.GetScale3D()),
+			LastBrushStamp.WorldNormal,
+			GEditor->GetEditorSubsystem<UPlacementModeSubsystem>()->GetModeSettingsObject());
+		FinalizedPlacementInfo.PreferredLevel = GEditor->GetEditorWorldContext().World()->GetCurrentLevel();
+
+		GetToolManager()->BeginUndoTransaction(NSLOCTEXT("PlacementMode", "SinglePlaceAsset", "Place Single Asset"));
+		PlacedElements = PlacementSubsystem->PlaceAsset(FinalizedPlacementInfo, PlacementOptions);
+		NotifyMovementStarted(PlacedElements);
+	}
+}
+
+void UPlacementModePlaceSingleTool::OnClickDrag(const FInputDeviceRay& DragPos)
+{
+	FVector TraceStartLocation = DragPos.WorldRay.Origin;
+	FVector TraceDirection = DragPos.WorldRay.Direction;
+	FVector TraceEndLocation = TraceStartLocation + (TraceDirection * HALF_WORLD_MAX);
+	FVector TraceIntersectionXY = FMath::LinePlaneIntersection(TraceStartLocation, TraceEndLocation, FPlane(LastBrushStamp.WorldPosition, LastBrushStamp.WorldNormal));
+
+	FVector CursorDirection;
+	float CursorDistance;
+	FVector MouseDelta = LastBrushStamp.WorldPosition - TraceIntersectionXY;
+	MouseDelta.ToDirectionAndLength(CursorDirection, CursorDistance);
+	if (CursorDirection.IsNearlyZero())
+	{
+		return;
+	}
+
+	const UAssetPlacementSettings* PlacementSettings = GEditor->GetEditorSubsystem<UPlacementModeSubsystem>()->GetModeSettingsObject();
+
+	// Update rotation based on mouse position
+	FTransform UpdatedTransform = FinalizeTransform(
+		FTransform(FRotationMatrix::MakeFromXZ(CursorDirection, PlacementInfo->FinalizedTransform.GetRotation().GetUpVector()).ToQuat(), LastBrushStamp.WorldPosition, PlacementInfo->FinalizedTransform.GetScale3D()),
+		LastBrushStamp.WorldNormal,
+		PlacementSettings);
+
+	// Update scale based on mouse position
+	FVector UpdatedScale = UpdatedTransform.GetScale3D();
+	if (PlacementSettings && (SinglePlaceSettings->ScalingType != EPlacementScaleToCursorType::None))
+	{
+		auto UpdateComponent = [PlacementSettings, CursorDistance, this](float InComponent) -> float
+		{
+			float Sign = 1.0f;
+			if (InComponent < 0.0f)
+			{
+				Sign = -1.0f;
+			}
+
+			FFloatInterval ScaleRange(FMath::Abs(InComponent), PlacementSettings->ScaleRange.Max);
+			float NewComponent = ScaleRange.Interpolate(FMath::Min(1.0f, (CursorDistance / BrushStampIndicator->BrushRadius)));
+			return NewComponent * Sign;
+		};
+
+		switch (PlacementSettings->ScalingType)
+		{
+			case EFoliageScaling::LockXY:
+			{
+				UpdatedScale.Z = UpdateComponent(UpdatedScale.Z);
+				break;
+			}
+			case EFoliageScaling::LockYZ:
+			{
+				UpdatedScale.X = UpdateComponent(UpdatedScale.X);
+				break;
+			}
+			case EFoliageScaling::LockXZ:
+			{
+				UpdatedScale.Y = UpdateComponent(UpdatedScale.Y);
+				break;
+			}
+			default:
+			{
+				UpdatedScale.X = UpdateComponent(UpdatedScale.X);
+				UpdatedScale.Y = UpdateComponent(UpdatedScale.Y);
+				UpdatedScale.Z = UpdateComponent(UpdatedScale.Z);
+				break;
+			}
+		}
+	}
+	UpdatedTransform.SetScale3D(UpdatedScale);
+
+	// Use the drag position and settings to update the scale and rotation of the placed elements.
+	UpdateElementTransforms(PlacedElements, UpdatedTransform, false);
+}
+
+void UPlacementModePlaceSingleTool::OnClickRelease(const FInputDeviceRay& ReleasePos)
+{
+	Super::OnClickRelease(ReleasePos);
+
+	NotifyMovementEnded(PlacedElements);
+	EnterTweakState(PlacedElements);
+	GetToolManager()->EndUndoTransaction();
+
+	ShutdownBrushStampIndicator();
+	PlacementInfo.Reset();
+	PlacedElements.Empty();
+}
+
+void UPlacementModePlaceSingleTool::OnBeginHover(const FInputDeviceRay& DevicePos)
+{
+	Super::OnBeginHover(DevicePos);
+
+	// Always regenerate the placement data when a hover sequence begins
+	PlacementInfo.Reset();
+	CreatePreviewElements(DevicePos);
+}
+
+bool UPlacementModePlaceSingleTool::OnUpdateHover(const FInputDeviceRay& DevicePos)
+{
+	if (!Super::OnUpdateHover(DevicePos))
+	{
+		return false;
+	}
+
+	// Update the preview elements
+	UpdatePreviewElements(DevicePos);
+	return true;
+}
+
+void UPlacementModePlaceSingleTool::OnEndHover()
+{
+	Super::OnEndHover();
+
+	// Destroy the preview elements we created.
+	DestroyPreviewElements();
+}
+
+FInputRayHit UPlacementModePlaceSingleTool::CanBeginClickDragSequence(const FInputDeviceRay& PressPos)
+{
+	if (!bIsTweaking)
+	{
+		return Super::CanBeginClickDragSequence(PressPos);
+	}
+	
+	return FInputRayHit();
+}
+
+FInputRayHit UPlacementModePlaceSingleTool::BeginHoverSequenceHitTest(const FInputDeviceRay& DevicePos)
+{
+	if (!bIsTweaking)
+	{
+		return Super::BeginHoverSequenceHitTest(DevicePos);
+	}
+
+	return FInputRayHit();
+}
+
+void UPlacementModePlaceSingleTool::GeneratePlacementData(const FInputDeviceRay& DevicePos)
+{
+	UPlacementModeSubsystem* PlacementSubsystem = GEditor->GetEditorSubsystem<UPlacementModeSubsystem>();
+	const UAssetPlacementSettings* PlacementSettings = PlacementSubsystem ? PlacementSubsystem->GetModeSettingsObject() : nullptr;
+	if (PlacementSettings)
+	{
+		TArrayView<const TObjectPtr<UPlacementPaletteClient>> CurrentPaletteItems = PlacementSettings->GetActivePaletteItems();
+		if (CurrentPaletteItems.Num())
+		{
+			int32 ItemIndex = FMath::RandHelper(CurrentPaletteItems.Num());
+			const UPlacementPaletteClient* ItemToPlace = CurrentPaletteItems[ItemIndex];
+			if (ItemToPlace)
+			{
+				FTransform TransformToUpdate(GenerateRandomRotation(PlacementSettings), LastBrushStamp.WorldPosition, GenerateRandomScale(PlacementSettings));
+
+				PlacementInfo = MakeUnique<FAssetPlacementInfo>();
+				PlacementInfo->AssetToPlace = ItemToPlace->AssetPath.TryLoad();
+				PlacementInfo->FactoryOverride = ItemToPlace->FactoryInterfaceClass;
+				PlacementInfo->ItemGuid = ItemToPlace->ClientGuid;
+				PlacementInfo->FinalizedTransform = FinalizeTransform(TransformToUpdate, LastBrushStamp.WorldNormal, PlacementSettings);
+				PlacementInfo->PreferredLevel = GEditor->GetEditorWorldContext().World()->GetCurrentLevel();
+				PlacementInfo->SettingsObject = ItemToPlace->SettingsObject;
+			}
+		}
+	}
+}
+
+void UPlacementModePlaceSingleTool::CreatePreviewElements(const FInputDeviceRay& DevicePos)
+{
+	SetupBrushStampIndicator();
+
+	// Place the preview elements from our stored info
+	if (!PlacementInfo)
+	{
+		GeneratePlacementData(DevicePos);
+	}
+
+	if (UPlacementSubsystem* PlacementSubsystem = GEditor->GetEditorSubsystem<UPlacementSubsystem>())
+	{
+		FPlacementOptions PlacementOptions;
+		PlacementOptions.bIsCreatingPreviewElements = true;
+		FAssetPlacementInfo InfoToPlace = *PlacementInfo;
+		InfoToPlace.FinalizedTransform = FinalizeTransform(
+			FTransform(PlacementInfo->FinalizedTransform.GetRotation(), LastBrushStamp.WorldPosition, PlacementInfo->FinalizedTransform.GetScale3D()),
+			LastBrushStamp.WorldNormal,
+			GEditor->GetEditorSubsystem<UPlacementModeSubsystem>()->GetModeSettingsObject());
+
+		PreviewElements = PlacementSubsystem->PlaceAsset(InfoToPlace, PlacementOptions);
+	}
+
+	NotifyMovementStarted(PreviewElements);
+}
+
+void UPlacementModePlaceSingleTool::UpdatePreviewElements(const FInputDeviceRay& DevicePos)
+{
+	// If we should have preview elements, but do not currently, go ahead and create them.
+	if ((PreviewElements.Num() == 0) && !BrushStampIndicator)
+	{
+		CreatePreviewElements(DevicePos);
+	}
+
+	// If we don't actually have any preview handles created, we don't need to update them, so go ahead and bail.
+	if (PreviewElements.Num() == 0)
+	{
+		return;
+	}
+
+	// Update the brush radius to be stable screen size based on the world to pixel scale at the last render update.
+	LastBrushStamp.Radius = 100.0f * LastBrushStampWorldToPixelScale;
+
+	FTransform UpdatedTransform(PlacementInfo->FinalizedTransform.GetRotation(), LastBrushStamp.WorldPosition, PlacementInfo->FinalizedTransform.GetScale3D());
+	UpdatedTransform = FinalizeTransform(UpdatedTransform, LastBrushStamp.WorldNormal, GEditor->GetEditorSubsystem<UPlacementModeSubsystem>()->GetModeSettingsObject());
+
+	UpdateElementTransforms(PreviewElements, UpdatedTransform);
+}
+
+void UPlacementModePlaceSingleTool::DestroyPreviewElements()
+{
+	if (IAssetEditorContextInterface* AssetEditorContext = GetToolManager()->GetContextObjectStore()->FindContext<IAssetEditorContextInterface>())
+	{
+		if (UTypedElementSelectionSet* SelectionSet = AssetEditorContext->GetMutableSelectionSet())
+		{
+			for (const FTypedElementHandle& PreviewElement : PreviewElements)
+			{
+				if (TTypedElement<ITypedElementWorldInterface> WorldInterfaceElement = UTypedElementRegistry::GetInstance()->GetElement<ITypedElementWorldInterface>(PreviewElement))
+				{
+					WorldInterfaceElement.NotifyMovementEnded();
+					WorldInterfaceElement.DeleteElement(WorldInterfaceElement.GetOwnerWorld(), SelectionSet, FTypedElementDeletionOptions());
+				}
+			}
+		}
+	}
+
+	PreviewElements.Empty();
+}
+
+void UPlacementModePlaceSingleTool::EnterTweakState(TArrayView<const FTypedElementHandle> InElementHandles)
+{
+	if (InElementHandles.Num() == 0 || !SinglePlaceSettings->bSelectAfterPlacing)
+	{
+		return;
+	}
+
+	if (IAssetEditorContextInterface* AssetEditorContext = GetToolManager()->GetContextObjectStore()->FindContext<IAssetEditorContextInterface>())
+	{
+		if (UTypedElementSelectionSet* SelectionSet = AssetEditorContext->GetMutableSelectionSet())
+		{
+			SelectionSet->SetSelection(InElementHandles, FTypedElementSelectionOptions());
+			bIsTweaking = true;
+		}
+	}
+}
+
+void UPlacementModePlaceSingleTool::ExitTweakState(bool bClearSelectionSet)
+{
+	if (bClearSelectionSet)
+	{
+		if (IAssetEditorContextInterface* AssetEditorContext = GetToolManager()->GetContextObjectStore()->FindContext<IAssetEditorContextInterface>())
+		{
+			if (UTypedElementSelectionSet* SelectionSet = AssetEditorContext->GetMutableSelectionSet())
+			{
+				SelectionSet->ClearSelection(FTypedElementSelectionOptions());
+			}
+		}
+	}
+
+	bIsTweaking = false;
+}
+
+void UPlacementModePlaceSingleTool::UpdateElementTransforms(TArrayView<const FTypedElementHandle> InElements, const FTransform& InTransform, bool bLocalTransform)
+{
+	// Update the transform positions for the preview elements.
+	for (const FTypedElementHandle& ElementHandle : InElements)
+	{
+		if (TTypedElement<ITypedElementWorldInterface> WorldInterfaceElement = UTypedElementRegistry::GetInstance()->GetElement<ITypedElementWorldInterface>(ElementHandle))
+		{
+			bLocalTransform ? WorldInterfaceElement.SetRelativeTransform(InTransform) : WorldInterfaceElement.SetWorldTransform(InTransform);
+			WorldInterfaceElement.NotifyMovementOngoing();
+		}
+	}
+}
+
+void UPlacementModePlaceSingleTool::NotifyMovementStarted(TArrayView<const FTypedElementHandle> InElements)
+{
+	// Notify Movement started
+	for (const FTypedElementHandle& ElementToNotify : InElements)
+	{
+		if (TTypedElement<ITypedElementWorldInterface> WorldInterfaceElement = UTypedElementRegistry::GetInstance()->GetElement<ITypedElementWorldInterface>(ElementToNotify))
+		{
+			WorldInterfaceElement.NotifyMovementStarted();
+		}
+	}
+}
+
+void UPlacementModePlaceSingleTool::NotifyMovementEnded(TArrayView<const FTypedElementHandle> InElements)
+{
+	for (const FTypedElementHandle& ElementToNotify : InElements)
+	{
+		if (TTypedElement<ITypedElementWorldInterface> WorldInterfaceElement = UTypedElementRegistry::GetInstance()->GetElement<ITypedElementWorldInterface>(ElementToNotify))
+		{
+			WorldInterfaceElement.NotifyMovementEnded();
+		}
+	}
+}
+
+void UPlacementModePlaceSingleTool::SetupRightClickMouseBehavior()
+{
+	ULocalClickDragInputBehavior* RightMouseBehavior = NewObject<ULocalClickDragInputBehavior>(this);
+	RightMouseBehavior->CanBeginClickDragFunc = [this](const FInputDeviceRay&) { return bShiftToggle ? FInputRayHit(1.0f) : FInputRayHit(); };
+	RightMouseBehavior->OnClickReleaseFunc = [this](const FInputDeviceRay&)
+	{
+		constexpr bool bClearSelection = true;
+		ExitTweakState(bClearSelection);
+		PlacementInfo.Reset();
+	};
+	RightMouseBehavior->SetDefaultPriority(FInputCapturePriority(-1));
+	RightMouseBehavior->SetUseRightMouseButton();
+	RightMouseBehavior->Initialize();
+	AddInputBehavior(RightMouseBehavior);
+}
diff --git a/Engine/Plugins/Experimental/AssetPlacementEdMode/Source/AssetPlacementEdMode/Private/Tools/PlacementPlaceSingleTool.h b/Engine/Plugins/Experimental/AssetPlacementEdMode/Source/AssetPlacementEdMode/Private/Tools/PlacementPlaceSingleTool.h
new file mode 100644
--- /dev/null
+++ b/Engine/Plugins/Experimental/AssetPlacementEdMode/Source/AssetPlacementEdMode/Private/Tools/PlacementPlaceSingleTool.h
@@ -0,0 +1,87 @@
+// Copyright Epic Games, Inc. All Rights Reserved.
+
+#pragma once
+
+#include "Tools/PlacementClickDragToolBase.h"
+
+#include "PlacementPlaceSingleTool.generated.h"
+
+struct FAssetPlacementInfo;
+
+UCLASS(Transient, MinimalAPI)
+class UPlacementModePlaceSingleToolBuilder : public UPlacementToolBuilderBase
+{
+	GENERATED_BODY()
+
+protected:
+	virtual UPlacementBrushToolBase* FactoryToolInstance(UObject* Outer) const override;
+};
+
+UENUM()
+enum class EPlacementScaleToCursorType : uint8
+{
+	/** Increases scale as the cursor moves outward from the placed asset. */
+	Positive,
+	/** Scale is unchanged by cursor movement for the placed asset. */
+	None,
+};
+
+UCLASS(config = EditorPerProjectUserSettings)
+class UPlacementModePlaceSingleToolSettings : public UInteractiveToolPropertySet
+{
+	GENERATED_BODY()
+
+public:
+	/** How the cursor movement should apply to scale once an asset is placed. Note that the maximum bound is controlled by scale settings of the mode. */
+	UPROPERTY(config, EditAnywhere, Category = "Single Place Settings")
+	EPlacementScaleToCursorType ScalingType = EPlacementScaleToCursorType::Positive;
+
+	/** If the tool should automatically select the last placed asset. */
+	UPROPERTY(config, EditAnywhere, Category = "Single Place Settings")
+	bool bSelectAfterPlacing = true;
+};
+
+UCLASS(Transient)
+class UPlacementModePlaceSingleTool : public UPlacementClickDragToolBase
+{
+	GENERATED_BODY()
+
+public:
+	UPlacementModePlaceSingleTool();
+	~UPlacementModePlaceSingleTool();
+	UPlacementModePlaceSingleTool(FVTableHelper& Helper);
+	
+	constexpr static TCHAR ToolName[] = TEXT("PlaceSingleTool");
+
+	virtual void Setup() override;
+	virtual void Shutdown(EToolShutdownType ShutdownType) override;
+	virtual void OnClickPress(const FInputDeviceRay& Ray) override;
+	virtual void OnClickDrag(const FInputDeviceRay& DragPos) override;
+	virtual void OnClickRelease(const FInputDeviceRay& ReleasePos) override;
+
+	virtual void OnBeginHover(const FInputDeviceRay& DevicePos) override;
+	virtual bool OnUpdateHover(const FInputDeviceRay& DevicePos) override;
+	virtual void OnEndHover() override;
+	virtual FInputRayHit CanBeginClickDragSequence(const FInputDeviceRay& PressPos) override;
+	virtual FInputRayHit BeginHoverSequenceHitTest(const FInputDeviceRay& DevicePos) override;
+
+protected:
+	void GeneratePlacementData(const FInputDeviceRay& DevicePos);
+	void CreatePreviewElements(const FInputDeviceRay& DevicePos);
+	void UpdatePreviewElements(const FInputDeviceRay& DevicePos);
+	void DestroyPreviewElements();
+	void EnterTweakState(TArrayView<const FTypedElementHandle> InElementHandles);
+	void ExitTweakState(bool bClearSelectionSet);
+	void UpdateElementTransforms(TArrayView<const FTypedElementHandle> InElements, const FTransform& InTransform, bool bLocalTransform = false);
+	void NotifyMovementStarted(TArrayView<const FTypedElementHandle> InElements);
+	void NotifyMovementEnded(TArrayView<const FTypedElementHandle> InElements);
+	void SetupRightClickMouseBehavior();
+
+	UPROPERTY(Transient)
+	TObjectPtr<UPlacementModePlaceSingleToolSettings> SinglePlaceSettings;
+
+	TUniquePtr<FAssetPlacementInfo> PlacementInfo;
+	TArray<FTypedElementHandle> PreviewElements;
+	TArray<FTypedElementHandle> PlacedElements;
+	bool bIsTweaking;
+};
diff --git a/Engine/Plugins/Experimental/AssetPlacementEdMode/Source/AssetPlacementEdMode/Private/Tools/PlacementPlaceTool.cpp b/Engine/Plugins/Experimental/AssetPlacementEdMode/Source/AssetPlacementEdMode/Private/Tools/PlacementPlaceTool.cpp
new file mode 100644
--- /dev/null
+++ b/Engine/Plugins/Experimental/AssetPlacementEdMode/Source/AssetPlacementEdMode/Private/Tools/PlacementPlaceTool.cpp
@@ -0,0 +1,148 @@
+// Copyright Epic Games, Inc. All Rights Reserved.
+
+#include "Tools/PlacementPlaceTool.h"
+#include "UObject/Object.h"
+#include "AssetPlacementSettings.h"
+#include "Subsystems/PlacementSubsystem.h"
+#include "Math/UnrealMathUtility.h"
+#include "Editor.h"
+#include "InstancedFoliage.h"
+#include "InteractiveToolManager.h"
+#include "Modes/PlacementModeSubsystem.h"
+#include "PlacementPaletteItem.h"
+
+#include "Instances/InstancedPlacementClientInfo.h"
+
+constexpr TCHAR UPlacementModePlacementTool::ToolName[];
+
+UPlacementBrushToolBase* UPlacementModePlacementToolBuilder::FactoryToolInstance(UObject* Outer) const
+{
+	return NewObject<UPlacementModePlacementTool>(Outer);
+}
+
+void UPlacementModePlacementTool::OnBeginDrag(const FRay& Ray)
+{
+	Super::OnBeginDrag(Ray);
+
+	GetToolManager()->BeginUndoTransaction(NSLOCTEXT("AssetPlacementEdMode", "PaintAssets", "Paint Assets"));
+}
+
+void UPlacementModePlacementTool::OnEndDrag(const FRay& Ray)
+{
+	GetToolManager()->EndUndoTransaction();
+
+	Super::OnEndDrag(Ray);
+}
+
+void UPlacementModePlacementTool::OnTick(float DeltaTime)
+{
+	if (!bInBrushStroke)
+	{
+		return;
+	}
+
+	UPlacementModeSubsystem* PlacementModeSubsystem = GEditor->GetEditorSubsystem<UPlacementModeSubsystem>();
+	const UAssetPlacementSettings* PlacementSettings = PlacementModeSubsystem ? PlacementModeSubsystem->GetModeSettingsObject() : nullptr;
+	if (BrushProperties && PlacementSettings)
+	{
+		TArrayView<const TObjectPtr<UPlacementPaletteClient>> PaletteItems = PlacementSettings->GetActivePaletteItems();
+		if (PaletteItems.Num())
+		{
+			FVector BrushLocation = LastBrushStamp.WorldPosition;
+
+			// Assume a default density of 100 * whatever the user has selected as brush size.
+			float DefaultDensity = 100.0f * BrushProperties->BrushSize;
+			// This is the total set of instances disregarding parameters like slope, height or layer.
+			float DesiredInstanceCountFloat = DefaultDensity /** todo: individual item's Density*/ * BrushProperties->BrushStrength;
+			// Allow a single instance with a random chance, if the brush is smaller than the density
+			int32 DesiredInstanceCount = DesiredInstanceCountFloat > 1.f ? FMath::RoundToInt(DesiredInstanceCountFloat) : FMath::FRand() < DesiredInstanceCountFloat ? 1 : 0;
+
+			// Todo: save the instance hash per tile for the editor, so that paints don't continually add instances and surpass desired density (partition actor?)
+			FFoliageInstanceHash PotentialInstanceHash(7);
+			TArray<FVector> PotentialInstanceLocations;
+
+			TArray<FAssetPlacementInfo> DesiredInfoSpawnLocations;
+			DesiredInfoSpawnLocations.Reserve(DesiredInstanceCount);
+			PotentialInstanceLocations.Empty(DesiredInstanceCount);
+			for (int32 Index = 0; Index < DesiredInstanceCount; ++Index)
+			{
+				FVector Start, End;
+				GetRandomVectorInBrush(Start, End);
+				FHitResult AdjustedHitResult;
+				if (FindHitResultWithStartAndEndTraceVectors(AdjustedHitResult, Start, End))
+				{
+					FVector SpawnLocation = AdjustedHitResult.ImpactPoint;
+					FVector SpawnNormal = AdjustedHitResult.ImpactNormal;
+
+					bool bValidInstance = true;
+					auto TempInstances = PotentialInstanceHash.GetInstancesOverlappingBox(FBox::BuildAABB(BrushLocation, FVector(BrushProperties->BrushRadius)));
+					for (int32 InstanceIndex : TempInstances)
+					{
+						if ((PotentialInstanceLocations[InstanceIndex] - BrushLocation).SizeSquared() < FMath::Square(BrushProperties->BrushRadius))
+						{
+							bValidInstance = false;
+							break;
+						}
+					}
+
+					if (!bValidInstance)
+					{
+						continue;
+					}
+
+					int32 PotentialIdx = PotentialInstanceLocations.Add(SpawnLocation);
+					PotentialInstanceHash.InsertInstance(SpawnLocation, PotentialIdx);
+
+					int32 ItemIndex = FMath::RandHelper(PaletteItems.Num());
+					if (const UPlacementPaletteClient* ItemToPlace = PaletteItems[ItemIndex])
+					{
+						FAssetPlacementInfo NewInfo;
+						NewInfo.AssetToPlace = ItemToPlace->AssetPath.TryLoad();
+						NewInfo.FactoryOverride = ItemToPlace->FactoryInterfaceClass;
+						NewInfo.ItemGuid = ItemToPlace->ClientGuid;
+						NewInfo.PreferredLevel = GEditor->GetEditorWorldContext().World()->GetCurrentLevel();
+						NewInfo.FinalizedTransform = GenerateTransformFromHitLocationAndNormal(SpawnLocation, SpawnNormal);
+						NewInfo.SettingsObject = ItemToPlace->SettingsObject;
+
+						DesiredInfoSpawnLocations.Emplace(NewInfo);
+					}
+				}
+			}
+
+			UPlacementSubsystem* PlacementSubsystem = GEditor->GetEditorSubsystem<UPlacementSubsystem>();
+			if (PlacementSubsystem)
+			{
+				FPlacementOptions PlacementOptions;
+				PlacementOptions.bPreferBatchPlacement = true;
+				PlacementOptions.InstancedPlacementGridGuid = PlacementSettings->GetActivePaletteGuid();
+				PlacementSubsystem->PlaceAssets(DesiredInfoSpawnLocations, PlacementOptions);
+			}
+		}
+	}
+}
+
+void UPlacementModePlacementTool::GetRandomVectorInBrush(FVector& OutStart, FVector& OutEnd)
+{
+	if (!BrushProperties)
+	{
+		return;
+	}
+	FVector BrushNormal = LastBrushStamp.WorldNormal;
+	FVector BrushLocation = LastBrushStamp.WorldPosition;
+	float BrushRadius = BrushProperties->BrushRadius;
+
+	// Find Rx and Ry inside the unit circle
+	float Ru = (2.f * FMath::FRand() - 1.f);
+	float Rv = (2.f * FMath::FRand() - 1.f) * FMath::Sqrt(1.f - FMath::Square(Ru));
+
+	// find random point in circle through brush location on the same plane to brush location hit surface normal
+	FVector U, V;
+	BrushNormal.FindBestAxisVectors(U, V);
+	FVector Point = Ru * U + Rv * V;
+
+	// find distance to surface of sphere brush from this point
+	FVector Rw = FMath::Sqrt(FMath::Max(1.f - (FMath::Square(Ru) + FMath::Square(Rv)), 0.001f)) * BrushNormal;
+
+	OutStart = BrushLocation + BrushRadius * (Point + Rw);
+	OutEnd = BrushLocation + BrushRadius * (Point - Rw);
+}
diff --git a/Engine/Plugins/Experimental/AssetPlacementEdMode/Source/AssetPlacementEdMode/Private/Tools/PlacementPlaceTool.h b/Engine/Plugins/Experimental/AssetPlacementEdMode/Source/AssetPlacementEdMode/Private/Tools/PlacementPlaceTool.h
new file mode 100644
--- /dev/null
+++ b/Engine/Plugins/Experimental/AssetPlacementEdMode/Source/AssetPlacementEdMode/Private/Tools/PlacementPlaceTool.h
@@ -0,0 +1,33 @@
+// Copyright Epic Games, Inc. All Rights Reserved.
+
+#pragma once
+
+#include "InteractiveToolBuilder.h"
+#include "Tools/PlacementBrushToolBase.h"
+
+#include "PlacementPlaceTool.generated.h"
+
+UCLASS(Transient, MinimalAPI)
+class UPlacementModePlacementToolBuilder : public UPlacementToolBuilderBase
+{
+	GENERATED_BODY()
+
+public:
+	virtual UPlacementBrushToolBase* FactoryToolInstance(UObject* Outer) const override;
+};
+
+UCLASS(MinimalAPI)
+class UPlacementModePlacementTool : public UPlacementBrushToolBase
+{
+	GENERATED_BODY()
+public:
+	constexpr static TCHAR ToolName[] = TEXT("PlaceTool");
+
+
+protected:
+	virtual void OnBeginDrag(const FRay& Ray) override;
+	virtual void OnEndDrag(const FRay& Ray) override;
+	virtual void OnTick(float DeltaTime) override;
+
+	void GetRandomVectorInBrush(FVector& OutStart, FVector& OutEnd);
+};
diff --git a/Engine/Plugins/Experimental/AssetPlacementEdMode/Source/AssetPlacementEdMode/Private/Tools/PlacementSelectTool.cpp b/Engine/Plugins/Experimental/AssetPlacementEdMode/Source/AssetPlacementEdMode/Private/Tools/PlacementSelectTool.cpp
new file mode 100644
--- /dev/null
+++ b/Engine/Plugins/Experimental/AssetPlacementEdMode/Source/AssetPlacementEdMode/Private/Tools/PlacementSelectTool.cpp
@@ -0,0 +1,29 @@
+// Copyright Epic Games, Inc. All Rights Reserved.
+
+#include "Tools/PlacementSelectTool.h"
+#include "UObject/Object.h"
+#include "ToolContextInterfaces.h"
+#include "InteractiveToolManager.h"
+
+constexpr TCHAR UPlacementModeSelectTool::ToolName[];
+
+bool UPlacementModeSelectToolBuilder::CanBuildTool(const FToolBuilderState& SceneState) const
+{
+	return true;
+}
+
+UPlacementBrushToolBase* UPlacementModeSelectToolBuilder::FactoryToolInstance(UObject* Outer) const
+{
+	return NewObject<UPlacementModeSelectTool>(Outer);
+}
+
+void UPlacementModeSelectTool::Setup()
+{
+	Super::Setup();
+	ShutdownBrushStampIndicator();
+}
+
+FInputRayHit UPlacementModeSelectTool::CanBeginClickDragSequence(const FInputDeviceRay& PressPos)
+{
+	return FInputRayHit();	// always fall-through to the EdMode or viewport handling for clicks.
+}
diff --git a/Engine/Plugins/Experimental/AssetPlacementEdMode/Source/AssetPlacementEdMode/Private/Tools/PlacementSelectTool.h b/Engine/Plugins/Experimental/AssetPlacementEdMode/Source/AssetPlacementEdMode/Private/Tools/PlacementSelectTool.h
new file mode 100644
--- /dev/null
+++ b/Engine/Plugins/Experimental/AssetPlacementEdMode/Source/AssetPlacementEdMode/Private/Tools/PlacementSelectTool.h
@@ -0,0 +1,32 @@
+// Copyright Epic Games, Inc. All Rights Reserved.
+
+#pragma once
+
+#include "Tools/PlacementBrushToolBase.h"
+#include "Tools/PlacementClickDragToolBase.h"
+
+#include "PlacementSelectTool.generated.h"
+
+UCLASS(Transient, MinimalAPI)
+class UPlacementModeSelectToolBuilder : public UPlacementToolBuilderBase
+{
+	GENERATED_BODY()
+
+public:
+	virtual bool CanBuildTool(const FToolBuilderState& SceneState) const override;
+
+protected:
+	virtual UPlacementBrushToolBase* FactoryToolInstance(UObject* Outer) const override;
+};
+
+UCLASS(MinimalAPI)
+class UPlacementModeSelectTool : public UPlacementClickDragToolBase
+{
+	GENERATED_BODY()
+
+public:
+	virtual void Setup() override;
+	virtual FInputRayHit CanBeginClickDragSequence(const FInputDeviceRay& PressPos) override;
+
+	constexpr static TCHAR ToolName[] = TEXT("SelectTool");
+};
diff --git a/Engine/Plugins/Experimental/AssetPlacementEdMode/Source/AssetPlacementEdMode/Public/AssetPlacementEdMode.h b/Engine/Plugins/Experimental/AssetPlacementEdMode/Source/AssetPlacementEdMode/Public/AssetPlacementEdMode.h
new file mode 100644
--- /dev/null
+++ b/Engine/Plugins/Experimental/AssetPlacementEdMode/Source/AssetPlacementEdMode/Public/AssetPlacementEdMode.h
@@ -0,0 +1,56 @@
+// Copyright Epic Games, Inc. All Rights Reserved.
+
+#pragma once
+
+#include "CoreMinimal.h"
+#include "Tools/DefaultEdMode.h"
+#include "Elements/Framework/TypedElementHandle.h"
+
+#include "AssetPlacementEdMode.generated.h"
+
+class UAssetPlacementSettings;
+
+UCLASS()
+class UAssetPlacementEdMode : public UEdModeDefault
+{
+	GENERATED_BODY()
+
+public:
+	constexpr static const TCHAR AssetPlacementEdModeID[] = TEXT("EM_AssetPlacementEdMode");
+
+	UAssetPlacementEdMode();
+	virtual ~UAssetPlacementEdMode();
+
+	////////////////
+	// UEdMode interface
+	virtual void Enter() override;
+	virtual void Exit() override;
+	virtual void CreateToolkit() override;
+	virtual bool UsesToolkits() const override;
+	virtual TMap<FName, TArray<TSharedPtr<FUICommandInfo>>> GetModeCommands() const override;
+	virtual void BindCommands() override;
+	virtual bool IsSelectionAllowed(AActor* InActor, bool bInSelection) const override;
+	virtual void OnToolStarted(UInteractiveToolManager* Manager, UInteractiveTool* Tool) override;
+	virtual void OnToolEnded(UInteractiveToolManager* Manager, UInteractiveTool* Tool) override;
+	// End of UEdMode interface
+	//////////////////
+
+	////////////////
+	// UBaseLegacyWidgetEdMode interface
+	virtual bool UsesPropertyWidgets() const override;
+	virtual bool ShouldDrawWidget() const override;
+	// End of UBaseLegacyWidgetEdMode interface
+	//////////////////
+
+	static bool IsEnabled();
+
+protected:
+	void ClearSelection();
+	bool HasAnyAssetsInPalette() const;
+	bool HasActiveSelection() const;
+	bool IsInSelectionTool() const;
+	void OnSMIsntancedElementsEnabledChanged();
+
+	bool bIsInSelectionTool;
+	TWeakObjectPtr<UAssetPlacementSettings> SettingsObjectAsPlacementSettings;
+};
diff --git a/Engine/Plugins/Experimental/AssetPlacementEdMode/Source/AssetPlacementEdMode/Public/AssetPlacementEdModeCommands.h b/Engine/Plugins/Experimental/AssetPlacementEdMode/Source/AssetPlacementEdMode/Public/AssetPlacementEdModeCommands.h
new file mode 100644
--- /dev/null
+++ b/Engine/Plugins/Experimental/AssetPlacementEdMode/Source/AssetPlacementEdMode/Public/AssetPlacementEdModeCommands.h
@@ -0,0 +1,25 @@
+// Copyright Epic Games, Inc. All Rights Reserved.
+
+#pragma once
+
+#include "CoreMinimal.h"
+#include "Framework/Commands/Commands.h"
+
+class FAssetPlacementEdModeCommands : public TCommands<FAssetPlacementEdModeCommands>
+{
+public:
+	FAssetPlacementEdModeCommands();
+
+	virtual void RegisterCommands() override;
+	static TMap<FName, TArray<TSharedPtr<FUICommandInfo>>> GetCommands();
+
+	TSharedPtr<FUICommandInfo> Select;
+	TSharedPtr<FUICommandInfo> SelectAll;
+	TSharedPtr<FUICommandInfo> LassoSelect;
+	TSharedPtr<FUICommandInfo> Place;
+	TSharedPtr<FUICommandInfo> PlaceSingle;
+	TSharedPtr<FUICommandInfo> Erase;
+
+protected:
+	TMap<FName, TArray<TSharedPtr<FUICommandInfo>>> Commands;
+};
diff --git a/Engine/Plugins/Experimental/AssetPlacementEdMode/Source/AssetPlacementEdMode/Public/AssetPlacementEdModeModule.h b/Engine/Plugins/Experimental/AssetPlacementEdMode/Source/AssetPlacementEdMode/Public/AssetPlacementEdModeModule.h
new file mode 100644
--- /dev/null
+++ b/Engine/Plugins/Experimental/AssetPlacementEdMode/Source/AssetPlacementEdMode/Public/AssetPlacementEdModeModule.h
@@ -0,0 +1,18 @@
+// Copyright Epic Games, Inc. All Rights Reserved.
+
+#pragma once
+
+#include "CoreMinimal.h"
+#include "Modules/ModuleInterface.h"
+
+class IAssetTypeActions;
+
+class FAssetPlacementEdMode : public IModuleInterface
+{
+public:
+	virtual void StartupModule() override;
+	virtual void ShutdownModule() override;
+
+protected:
+	TSharedPtr<IAssetTypeActions> PaletteAssetActions;
+};
diff --git a/Engine/Plugins/Experimental/AssetPlacementEdMode/Source/AssetPlacementEdMode/Public/AssetPlacementEdModeToolkit.h b/Engine/Plugins/Experimental/AssetPlacementEdMode/Source/AssetPlacementEdMode/Public/AssetPlacementEdModeToolkit.h
new file mode 100644
--- /dev/null
+++ b/Engine/Plugins/Experimental/AssetPlacementEdMode/Source/AssetPlacementEdMode/Public/AssetPlacementEdModeToolkit.h
@@ -0,0 +1,36 @@
+// Copyright Epic Games, Inc. All Rights Reserved.
+
+#pragma once
+
+#include "CoreMinimal.h"
+#include "Toolkits/BaseToolkit.h"
+
+/**
+ * This FModeToolkit just creates a basic UI panel that allows various InteractiveTools to
+ * be initialized, and a DetailsView used to show properties of the active Tool.
+ */
+class FAssetPlacementEdModeToolkit : public FModeToolkit
+{
+public:
+	FAssetPlacementEdModeToolkit();
+	
+	// FModeToolkit interface 
+	virtual void GetToolPaletteNames(TArray<FName>& PaletteNames) const override;
+
+	// IToolkit interface 
+	virtual FName GetToolkitFName() const override;
+	virtual FText GetBaseToolkitName() const override;
+	virtual void InvokeUI() override;
+protected:
+	virtual void RequestModeUITabs() override;
+	TSharedRef<SDockTab> CreateAssetPalette(const FSpawnTabArgs& Args);
+	TSharedRef<SDockTab> CreatePaletteItemDetailsView(const FSpawnTabArgs& Args);
+	EVisibility GetDetailsViewHintTextVisibility() const;
+
+private:
+	TWeakPtr<SDockTab> AssetPaletteTab;
+	TWeakPtr<SDockTab> PaletteItemDetailsViewTab;
+	FMinorTabConfig AssetPaletteInfo;
+	FMinorTabConfig PaletteItemDetailsViewInfo;
+	TSharedPtr<IDetailsView> PaletteItemDetailsView;
+};
diff --git a/Engine/Source/Editor/AssetPlacementEdMode/AssetPlacementEdMode.Build.cs b/Engine/Source/Editor/AssetPlacementEdMode/AssetPlacementEdMode.Build.cs
deleted file mode 100644
diff --git a/Engine/Source/Editor/AssetPlacementEdMode/Private/AssetPlacementEdMode.cpp b/Engine/Source/Editor/AssetPlacementEdMode/Private/AssetPlacementEdMode.cpp
deleted file mode 100644
diff --git a/Engine/Source/Editor/AssetPlacementEdMode/Private/AssetPlacementEdModeCommands.cpp b/Engine/Source/Editor/AssetPlacementEdMode/Private/AssetPlacementEdModeCommands.cpp
deleted file mode 100644
diff --git a/Engine/Source/Editor/AssetPlacementEdMode/Private/AssetPlacementEdModeModule.cpp b/Engine/Source/Editor/AssetPlacementEdMode/Private/AssetPlacementEdModeModule.cpp
deleted file mode 100644
diff --git a/Engine/Source/Editor/AssetPlacementEdMode/Private/AssetPlacementEdModeStyle.cpp b/Engine/Source/Editor/AssetPlacementEdMode/Private/AssetPlacementEdModeStyle.cpp
deleted file mode 100644
diff --git a/Engine/Source/Editor/AssetPlacementEdMode/Private/AssetPlacementEdModeStyle.h b/Engine/Source/Editor/AssetPlacementEdMode/Private/AssetPlacementEdModeStyle.h
deleted file mode 100644
diff --git a/Engine/Source/Editor/AssetPlacementEdMode/Private/AssetPlacementEdModeToolkit.cpp b/Engine/Source/Editor/AssetPlacementEdMode/Private/AssetPlacementEdModeToolkit.cpp
deleted file mode 100644
diff --git a/Engine/Source/Editor/AssetPlacementEdMode/Private/AssetPlacementPaletteItem.cpp b/Engine/Source/Editor/AssetPlacementEdMode/Private/AssetPlacementPaletteItem.cpp
deleted file mode 100644
diff --git a/Engine/Source/Editor/AssetPlacementEdMode/Private/AssetPlacementPaletteItem.h b/Engine/Source/Editor/AssetPlacementEdMode/Private/AssetPlacementPaletteItem.h
deleted file mode 100644
diff --git a/Engine/Source/Editor/AssetPlacementEdMode/Private/AssetPlacementSettings.cpp b/Engine/Source/Editor/AssetPlacementEdMode/Private/AssetPlacementSettings.cpp
deleted file mode 100644
diff --git a/Engine/Source/Editor/AssetPlacementEdMode/Private/AssetPlacementSettings.h b/Engine/Source/Editor/AssetPlacementEdMode/Private/AssetPlacementSettings.h
deleted file mode 100644
diff --git a/Engine/Source/Editor/AssetPlacementEdMode/Private/AssetTypeActions_PlacementPalette.cpp b/Engine/Source/Editor/AssetPlacementEdMode/Private/AssetTypeActions_PlacementPalette.cpp
deleted file mode 100644
diff --git a/Engine/Source/Editor/AssetPlacementEdMode/Private/AssetTypeActions_PlacementPalette.h b/Engine/Source/Editor/AssetPlacementEdMode/Private/AssetTypeActions_PlacementPalette.h
deleted file mode 100644
diff --git a/Engine/Source/Editor/AssetPlacementEdMode/Private/Modes/PlacementModeSubsystem.cpp b/Engine/Source/Editor/AssetPlacementEdMode/Private/Modes/PlacementModeSubsystem.cpp
deleted file mode 100644
diff --git a/Engine/Source/Editor/AssetPlacementEdMode/Private/Modes/PlacementModeSubsystem.h b/Engine/Source/Editor/AssetPlacementEdMode/Private/Modes/PlacementModeSubsystem.h
deleted file mode 100644
diff --git a/Engine/Source/Editor/AssetPlacementEdMode/Private/PlacementPaletteAsset.cpp b/Engine/Source/Editor/AssetPlacementEdMode/Private/PlacementPaletteAsset.cpp
deleted file mode 100644
diff --git a/Engine/Source/Editor/AssetPlacementEdMode/Private/PlacementPaletteAsset.h b/Engine/Source/Editor/AssetPlacementEdMode/Private/PlacementPaletteAsset.h
deleted file mode 100644
diff --git a/Engine/Source/Editor/AssetPlacementEdMode/Private/PlacementPaletteItem.cpp b/Engine/Source/Editor/AssetPlacementEdMode/Private/PlacementPaletteItem.cpp
deleted file mode 100644
diff --git a/Engine/Source/Editor/AssetPlacementEdMode/Private/PlacementPaletteItem.h b/Engine/Source/Editor/AssetPlacementEdMode/Private/PlacementPaletteItem.h
deleted file mode 100644
diff --git a/Engine/Source/Editor/AssetPlacementEdMode/Private/SAssetPlacementPalette.cpp b/Engine/Source/Editor/AssetPlacementEdMode/Private/SAssetPlacementPalette.cpp
deleted file mode 100644
diff --git a/Engine/Source/Editor/AssetPlacementEdMode/Private/SAssetPlacementPalette.h b/Engine/Source/Editor/AssetPlacementEdMode/Private/SAssetPlacementPalette.h
deleted file mode 100644
diff --git a/Engine/Source/Editor/AssetPlacementEdMode/Private/Tools/PlacementBrushToolBase.cpp b/Engine/Source/Editor/AssetPlacementEdMode/Private/Tools/PlacementBrushToolBase.cpp
deleted file mode 100644
diff --git a/Engine/Source/Editor/AssetPlacementEdMode/Private/Tools/PlacementBrushToolBase.h b/Engine/Source/Editor/AssetPlacementEdMode/Private/Tools/PlacementBrushToolBase.h
deleted file mode 100644
diff --git a/Engine/Source/Editor/AssetPlacementEdMode/Private/Tools/PlacementClickDragToolBase.cpp b/Engine/Source/Editor/AssetPlacementEdMode/Private/Tools/PlacementClickDragToolBase.cpp
deleted file mode 100644
diff --git a/Engine/Source/Editor/AssetPlacementEdMode/Private/Tools/PlacementClickDragToolBase.h b/Engine/Source/Editor/AssetPlacementEdMode/Private/Tools/PlacementClickDragToolBase.h
deleted file mode 100644
diff --git a/Engine/Source/Editor/AssetPlacementEdMode/Private/Tools/PlacementEraseTool.cpp b/Engine/Source/Editor/AssetPlacementEdMode/Private/Tools/PlacementEraseTool.cpp
deleted file mode 100644
diff --git a/Engine/Source/Editor/AssetPlacementEdMode/Private/Tools/PlacementEraseTool.h b/Engine/Source/Editor/AssetPlacementEdMode/Private/Tools/PlacementEraseTool.h
deleted file mode 100644
diff --git a/Engine/Source/Editor/AssetPlacementEdMode/Private/Tools/PlacementFillTool.cpp b/Engine/Source/Editor/AssetPlacementEdMode/Private/Tools/PlacementFillTool.cpp
deleted file mode 100644
diff --git a/Engine/Source/Editor/AssetPlacementEdMode/Private/Tools/PlacementFillTool.h b/Engine/Source/Editor/AssetPlacementEdMode/Private/Tools/PlacementFillTool.h
deleted file mode 100644
diff --git a/Engine/Source/Editor/AssetPlacementEdMode/Private/Tools/PlacementLassoSelectTool.cpp b/Engine/Source/Editor/AssetPlacementEdMode/Private/Tools/PlacementLassoSelectTool.cpp
deleted file mode 100644
diff --git a/Engine/Source/Editor/AssetPlacementEdMode/Private/Tools/PlacementLassoSelectTool.h b/Engine/Source/Editor/AssetPlacementEdMode/Private/Tools/PlacementLassoSelectTool.h
deleted file mode 100644
diff --git a/Engine/Source/Editor/AssetPlacementEdMode/Private/Tools/PlacementPlaceSingleTool.cpp b/Engine/Source/Editor/AssetPlacementEdMode/Private/Tools/PlacementPlaceSingleTool.cpp
deleted file mode 100644
diff --git a/Engine/Source/Editor/AssetPlacementEdMode/Private/Tools/PlacementPlaceSingleTool.h b/Engine/Source/Editor/AssetPlacementEdMode/Private/Tools/PlacementPlaceSingleTool.h
deleted file mode 100644
diff --git a/Engine/Source/Editor/AssetPlacementEdMode/Private/Tools/PlacementPlaceTool.cpp b/Engine/Source/Editor/AssetPlacementEdMode/Private/Tools/PlacementPlaceTool.cpp
deleted file mode 100644
diff --git a/Engine/Source/Editor/AssetPlacementEdMode/Private/Tools/PlacementPlaceTool.h b/Engine/Source/Editor/AssetPlacementEdMode/Private/Tools/PlacementPlaceTool.h
deleted file mode 100644
diff --git a/Engine/Source/Editor/AssetPlacementEdMode/Private/Tools/PlacementSelectTool.cpp b/Engine/Source/Editor/AssetPlacementEdMode/Private/Tools/PlacementSelectTool.cpp
deleted file mode 100644
diff --git a/Engine/Source/Editor/AssetPlacementEdMode/Private/Tools/PlacementSelectTool.h b/Engine/Source/Editor/AssetPlacementEdMode/Private/Tools/PlacementSelectTool.h
deleted file mode 100644
diff --git a/Engine/Source/Editor/AssetPlacementEdMode/Public/AssetPlacementEdMode.h b/Engine/Source/Editor/AssetPlacementEdMode/Public/AssetPlacementEdMode.h
deleted file mode 100644
diff --git a/Engine/Source/Editor/AssetPlacementEdMode/Public/AssetPlacementEdModeCommands.h b/Engine/Source/Editor/AssetPlacementEdMode/Public/AssetPlacementEdModeCommands.h
deleted file mode 100644
diff --git a/Engine/Source/Editor/AssetPlacementEdMode/Public/AssetPlacementEdModeModule.h b/Engine/Source/Editor/AssetPlacementEdMode/Public/AssetPlacementEdModeModule.h
deleted file mode 100644
diff --git a/Engine/Source/Editor/AssetPlacementEdMode/Public/AssetPlacementEdModeToolkit.h b/Engine/Source/Editor/AssetPlacementEdMode/Public/AssetPlacementEdModeToolkit.h
deleted file mode 100644
diff --git a/Engine/Source/Editor/LevelAssetEditor/LevelAssetEditor.Build.cs b/Engine/Source/Editor/LevelAssetEditor/LevelAssetEditor.Build.cs
--- a/Engine/Source/Editor/LevelAssetEditor/LevelAssetEditor.Build.cs
+++ b/Engine/Source/Editor/LevelAssetEditor/LevelAssetEditor.Build.cs
@@ -56,7 +56,6 @@ public class LevelAssetEditor : ModuleRules
 		DynamicallyLoadedModuleNames.AddRange(
 			new string[]
 			{
-				"AssetPlacementEdMode",
 				// ... add any modules that your module loads dynamically here ...
 			}
 			);
diff --git a/Engine/Source/Editor/LevelAssetEditor/Private/LevelAssetEditor.cpp b/Engine/Source/Editor/LevelAssetEditor/Private/LevelAssetEditor.cpp
--- a/Engine/Source/Editor/LevelAssetEditor/Private/LevelAssetEditor.cpp
+++ b/Engine/Source/Editor/LevelAssetEditor/Private/LevelAssetEditor.cpp
@@ -27,8 +27,6 @@ void FLevelAssetEditorModule::StartupModule()
 		FIsActionButtonVisible::CreateRaw(this, &FLevelAssetEditorModule::IsEnabled));
 
 	UToolMenus::RegisterStartupCallback(FSimpleMulticastDelegate::FDelegate::CreateRaw(this, &FLevelAssetEditorModule::RegisterMenus));
-
-	FModuleManager::Get().LoadModule("AssetPlacementEdMode");
 }
 
 void FLevelAssetEditorModule::ShutdownModule()
diff --git a/Engine/Source/Editor/UnrealEd/Classes/Settings/LevelEditorMiscSettings.h b/Engine/Source/Editor/UnrealEd/Classes/Settings/LevelEditorMiscSettings.h
--- a/Engine/Source/Editor/UnrealEd/Classes/Settings/LevelEditorMiscSettings.h
+++ b/Engine/Source/Editor/UnrealEd/Classes/Settings/LevelEditorMiscSettings.h
@@ -47,10 +47,6 @@ public:
 	UPROPERTY(EditAnywhere, config, AdvancedDisplay, Category = Editing)
 	bool bEnableExperimentalLevelEditor = false;
 
-	/** If enabled, the experimental asset placement mode will be available in the level editor */
-	UPROPERTY(EditAnywhere, config, AdvancedDisplay, Category = Editing)
-	bool bEnableAssetPlacementMode = false;
-
 	/** If enabled, will avoid relabeling actors in UUnrealEdEngine::edactPasteSelected */
 	UPROPERTY(EditAnywhere, config, Category = Editing, meta = (DisplayName = "Avoid Actor Relabel on Paste Selected"))
 	uint32 bAvoidRelabelOnPasteSelected:1;
diff --git a/Engine/Source/Editor/UnrealEd/Private/Settings/SettingsClasses.cpp b/Engine/Source/Editor/UnrealEd/Private/Settings/SettingsClasses.cpp
--- a/Engine/Source/Editor/UnrealEd/Private/Settings/SettingsClasses.cpp
+++ b/Engine/Source/Editor/UnrealEd/Private/Settings/SettingsClasses.cpp
@@ -429,14 +429,6 @@ void ULevelEditorMiscSettings::PostEditChangeProperty( struct FPropertyChangedEv
 		FNavigationSystem::SetNavigationAutoUpdateEnabled(bNavigationAutoUpdate, EditorContext.World()->GetNavigationSystem());
 	}
 
-	if (Name == GET_MEMBER_NAME_CHECKED(ULevelEditorMiscSettings, bEnableAssetPlacementMode))
-	{
-		if (UAssetEditorSubsystem* AssetEditorSubsystem = GEditor->GetEditorSubsystem<UAssetEditorSubsystem>())
-		{
-			AssetEditorSubsystem->OnEditorModesChanged().Broadcast();
-		}
-	}
-
 	if (!FUnrealEdMisc::Get().IsDeletePreferences())
 	{
 		SaveConfig();
diff --git a/Engine/Source/Editor/UnrealEd/Private/Subsystems/AssetEditorSubsystem.cpp b/Engine/Source/Editor/UnrealEd/Private/Subsystems/AssetEditorSubsystem.cpp
--- a/Engine/Source/Editor/UnrealEd/Private/Subsystems/AssetEditorSubsystem.cpp
+++ b/Engine/Source/Editor/UnrealEd/Private/Subsystems/AssetEditorSubsystem.cpp
@@ -32,6 +32,8 @@
 #include "Tools/LegacyEdMode.h"
 #include "ProfilingDebugging/StallDetector.h"
 
+#include "Elements/SMInstance/SMInstanceElementData.h" // For SMInstanceElementDataUtil::SMInstanceElementsEnabled
+
 
 #define LOCTEXT_NAMESPACE "AssetEditorSubsystem"
 
@@ -55,6 +57,8 @@ void UAssetEditorSubsystem::Initialize(FSubsystemCollectionBase& Collection)
 	GEditor->OnEditorClose().AddUObject(this, &UAssetEditorSubsystem::OnEditorClose);
 	FCoreDelegates::OnEnginePreExit.AddUObject(this, &UAssetEditorSubsystem::UnregisterEditorModes);
 	FCoreDelegates::OnPostEngineInit.AddUObject(this, &UAssetEditorSubsystem::RegisterEditorModes);
+
+	SMInstanceElementDataUtil::OnSMInstanceElementsEnabledChanged().AddUObject(this, &UAssetEditorSubsystem::OnSMInstanceElementsEnabled);
 }
 
 void UAssetEditorSubsystem::Deinitialize()
@@ -63,6 +67,7 @@ void UAssetEditorSubsystem::Deinitialize()
 	GEditor->OnEditorClose().RemoveAll(this);
 	FCoreDelegates::OnEnginePreExit.RemoveAll(this);
 	FCoreDelegates::OnPostEngineInit.RemoveAll(this);
+	SMInstanceElementDataUtil::OnSMInstanceElementsEnabledChanged().RemoveAll(this);
 
 	// Don't attempt to report usage stats if analytics isn't available
 	if (FEngineAnalytics::IsAvailable())
@@ -1013,5 +1018,10 @@ void UAssetEditorSubsystem::UnregisterEditorModes()
 	EditorModes.Empty();
 }
 
-#undef LOCTEXT_NAMESPACE
+void UAssetEditorSubsystem::OnSMInstanceElementsEnabled()
+{
+	// Let the modes know that SM instance elements may have been enabled or disabled and update state accordingly
+	OnEditorModesChanged().Broadcast();
+}
 
+#undef LOCTEXT_NAMESPACE
diff --git a/Engine/Source/Editor/UnrealEd/Public/Subsystems/AssetEditorSubsystem.h b/Engine/Source/Editor/UnrealEd/Public/Subsystems/AssetEditorSubsystem.h
--- a/Engine/Source/Editor/UnrealEd/Public/Subsystems/AssetEditorSubsystem.h
+++ b/Engine/Source/Editor/UnrealEd/Public/Subsystems/AssetEditorSubsystem.h
@@ -248,6 +248,7 @@ private:
 
 	void RegisterEditorModes();
 	void UnregisterEditorModes();
+	void OnSMInstanceElementsEnabled();
 
 private:
 
diff --git a/Engine/Source/Runtime/Engine/Private/Elements/SMInstance/SMInstanceElementData.cpp b/Engine/Source/Runtime/Engine/Private/Elements/SMInstance/SMInstanceElementData.cpp
--- a/Engine/Source/Runtime/Engine/Private/Elements/SMInstance/SMInstanceElementData.cpp
+++ b/Engine/Source/Runtime/Engine/Private/Elements/SMInstance/SMInstanceElementData.cpp
@@ -19,9 +19,16 @@ static FAutoConsoleVariableRef CVarEnableSMInstanceElements(
 	FConsoleVariableDelegate::CreateLambda([](IConsoleVariable* InVariable)
 	{
 		FGlobalComponentRecreateRenderStateContext Context;
+		SMInstanceElementDataUtil::OnSMInstanceElementsEnabledChanged().Broadcast();
 	})
 );
 
+FSimpleMulticastDelegate& OnSMInstanceElementsEnabledChanged()
+{
+	static FSimpleMulticastDelegate OnInstanceElementsEnabledChanged;
+	return OnInstanceElementsEnabledChanged;
+}
+
 bool SMInstanceElementsEnabled()
 {
 	return GEnableSMInstanceElements != 0;
diff --git a/Engine/Source/Runtime/Engine/Public/Elements/SMInstance/SMInstanceElementData.h b/Engine/Source/Runtime/Engine/Public/Elements/SMInstance/SMInstanceElementData.h
--- a/Engine/Source/Runtime/Engine/Public/Elements/SMInstance/SMInstanceElementData.h
+++ b/Engine/Source/Runtime/Engine/Public/Elements/SMInstance/SMInstanceElementData.h
@@ -36,6 +36,8 @@ namespace SMInstanceElementDataUtil
  */
 ENGINE_API bool SMInstanceElementsEnabled();
 
+ENGINE_API FSimpleMulticastDelegate& OnSMInstanceElementsEnabledChanged();
+
 /**
  * Get the static mesh instance manager for the given instance.
  * @return The static mesh instance manager, or null if this instance cannot be managed.
