Fix issue where the supplied actor factory was not making it through to the actor spawn in placement subsystem

diff --git a/Engine/Plugins/Experimental/AssetPlacementEdMode/Source/AssetPlacementEdMode/Private/AssetPlacementSettings.cpp b/Engine/Plugins/Experimental/AssetPlacementEdMode/Source/AssetPlacementEdMode/Private/AssetPlacementSettings.cpp
--- a/Engine/Plugins/Experimental/AssetPlacementEdMode/Source/AssetPlacementEdMode/Private/AssetPlacementSettings.cpp
+++ b/Engine/Plugins/Experimental/AssetPlacementEdMode/Source/AssetPlacementEdMode/Private/AssetPlacementSettings.cpp
@@ -165,13 +165,12 @@ bool UAssetPlacementSettings::DoesActivePaletteSupportElement(const FTypedElemen
 
 			// The current implementation of the asset data interface for actors requires that individual actors report on assets contained within their components.
 			// Not all elements (legacy via actors) do this reliably, so additionally check the supplied factory for a match. 
-			if (!Item->FactoryInterfaceClass.Get())
+			if (!Item->FactoryInterface.GetInterface())
 			{
 				continue;
 			}
 
-			TScriptInterface<IAssetFactoryInterface> FactoryInterface = GEditor->GetEditorSubsystem<UPlacementSubsystem>()->GetAssetFactoryFromFactoryClass(Item->FactoryInterfaceClass);
-			FAssetData FoundAssetDataFromFactory = FactoryInterface->GetAssetDataFromElementHandle(InElementToCheck);
+			FAssetData FoundAssetDataFromFactory = Item->FactoryInterface->GetAssetDataFromElementHandle(InElementToCheck);
 			if (FoundAssetDataFromFactory.ToSoftObjectPath() == Item->AssetPath)
 			{
 				return true;
diff --git a/Engine/Plugins/Experimental/AssetPlacementEdMode/Source/AssetPlacementEdMode/Private/PlacementPaletteItem.h b/Engine/Plugins/Experimental/AssetPlacementEdMode/Source/AssetPlacementEdMode/Private/PlacementPaletteItem.h
--- a/Engine/Plugins/Experimental/AssetPlacementEdMode/Source/AssetPlacementEdMode/Private/PlacementPaletteItem.h
+++ b/Engine/Plugins/Experimental/AssetPlacementEdMode/Source/AssetPlacementEdMode/Private/PlacementPaletteItem.h
@@ -8,7 +8,7 @@
 
 #include "PlacementPaletteItem.generated.h"
 
-class UAssetFactoryInterface;
+class IAssetFactoryInterface;
 class UInstancedPlacemenClientSettings;
 
 UCLASS(NotPlaceable)
@@ -27,7 +27,7 @@ public:
 	TObjectPtr<UInstancedPlacemenClientSettings> SettingsObject = nullptr;
 
 	UPROPERTY()
-	TSubclassOf<UAssetFactoryInterface> FactoryInterfaceClass;
+	TScriptInterface<IAssetFactoryInterface> FactoryInterface;
 
 	virtual void PostEditChangeProperty(struct FPropertyChangedEvent& PropertyChangedEvent) override;
 	virtual void PostLoad() override;
diff --git a/Engine/Plugins/Experimental/AssetPlacementEdMode/Source/AssetPlacementEdMode/Private/Tools/PlacementPlaceSingleTool.cpp b/Engine/Plugins/Experimental/AssetPlacementEdMode/Source/AssetPlacementEdMode/Private/Tools/PlacementPlaceSingleTool.cpp
--- a/Engine/Plugins/Experimental/AssetPlacementEdMode/Source/AssetPlacementEdMode/Private/Tools/PlacementPlaceSingleTool.cpp
+++ b/Engine/Plugins/Experimental/AssetPlacementEdMode/Source/AssetPlacementEdMode/Private/Tools/PlacementPlaceSingleTool.cpp
@@ -248,7 +248,7 @@ void UPlacementModePlaceSingleTool::GeneratePlacementData(const FInputDeviceRay&
 
 				PlacementInfo = MakeUnique<FAssetPlacementInfo>();
 				PlacementInfo->AssetToPlace = ItemToPlace->AssetPath.TryLoad();
-				PlacementInfo->FactoryOverride = ItemToPlace->FactoryInterfaceClass;
+				PlacementInfo->FactoryOverride = ItemToPlace->FactoryInterface;
 				PlacementInfo->ItemGuid = ItemToPlace->ClientGuid;
 				PlacementInfo->FinalizedTransform = FinalizeTransform(TransformToUpdate, LastBrushStamp.WorldNormal, PlacementSettings);
 				PlacementInfo->PreferredLevel = GEditor->GetEditorWorldContext().World()->GetCurrentLevel();
diff --git a/Engine/Plugins/Experimental/AssetPlacementEdMode/Source/AssetPlacementEdMode/Private/Tools/PlacementPlaceTool.cpp b/Engine/Plugins/Experimental/AssetPlacementEdMode/Source/AssetPlacementEdMode/Private/Tools/PlacementPlaceTool.cpp
--- a/Engine/Plugins/Experimental/AssetPlacementEdMode/Source/AssetPlacementEdMode/Private/Tools/PlacementPlaceTool.cpp
+++ b/Engine/Plugins/Experimental/AssetPlacementEdMode/Source/AssetPlacementEdMode/Private/Tools/PlacementPlaceTool.cpp
@@ -98,7 +98,7 @@ void UPlacementModePlacementTool::OnTick(float DeltaTime)
 					{
 						FAssetPlacementInfo NewInfo;
 						NewInfo.AssetToPlace = ItemToPlace->AssetPath.TryLoad();
-						NewInfo.FactoryOverride = ItemToPlace->FactoryInterfaceClass;
+						NewInfo.FactoryOverride = ItemToPlace->FactoryInterface;
 						NewInfo.ItemGuid = ItemToPlace->ClientGuid;
 						NewInfo.PreferredLevel = GEditor->GetEditorWorldContext().World()->GetCurrentLevel();
 						NewInfo.FinalizedTransform = GenerateTransformFromHitLocationAndNormal(SpawnLocation, SpawnNormal);
diff --git a/Engine/Source/Editor/EditorFramework/Private/Subsystems/PlacementSubsystem.cpp b/Engine/Source/Editor/EditorFramework/Private/Subsystems/PlacementSubsystem.cpp
--- a/Engine/Source/Editor/EditorFramework/Private/Subsystems/PlacementSubsystem.cpp
+++ b/Engine/Source/Editor/EditorFramework/Private/Subsystems/PlacementSubsystem.cpp
@@ -37,7 +37,7 @@ TArray<FTypedElementHandle> UPlacementSubsystem::PlaceAssets(TArrayView<const FA
 	for (const FAssetPlacementInfo& PlacementInfo : InPlacementInfos)
 	{
 		const FAssetData& AssetData = PlacementInfo.AssetToPlace;
-		TScriptInterface<IAssetFactoryInterface> FactoryInterface = GetAssetFactoryFromFactoryClass(PlacementInfo.FactoryOverride);
+		TScriptInterface<IAssetFactoryInterface> FactoryInterface = PlacementInfo.FactoryOverride;
 		if (!FactoryInterface)
 		{
 			FactoryInterface = FindAssetFactoryFromAssetData(AssetData);
@@ -119,7 +119,7 @@ void UPlacementSubsystem::RegisterPlacementFactories()
 	PlacementFactoriesRegistered.Broadcast();
 }
 
-TScriptInterface<IAssetFactoryInterface> UPlacementSubsystem::GetAssetFactoryFromFactoryClass(TSubclassOf<UAssetFactoryInterface> InFactoryInterfaceClass) const
+TScriptInterface<IAssetFactoryInterface> UPlacementSubsystem::GetAssetFactoryFromFactoryClass(TSubclassOf<UClass> InFactoryInterfaceClass) const
 {
 	if (InFactoryInterfaceClass)
 	{
diff --git a/Engine/Source/Editor/EditorFramework/Public/Subsystems/PlacementSubsystem.h b/Engine/Source/Editor/EditorFramework/Public/Subsystems/PlacementSubsystem.h
--- a/Engine/Source/Editor/EditorFramework/Public/Subsystems/PlacementSubsystem.h
+++ b/Engine/Source/Editor/EditorFramework/Public/Subsystems/PlacementSubsystem.h
@@ -36,7 +36,7 @@ struct EDITORFRAMEWORK_API FAssetPlacementInfo
 
 	// If set, will use the given factory to place the asset, instead of allowing the placement subsystem to determine which factory to use.
 	UPROPERTY()
-	TSubclassOf<UAssetFactoryInterface> FactoryOverride;
+	TScriptInterface<IAssetFactoryInterface> FactoryOverride;
 
 	/**
 	 * The Guid which corresponds to the item that should be placed.
@@ -108,7 +108,7 @@ public:
 
 	FSimpleMulticastDelegate& OnPlacementFactoriesRegistered();
 
-	TScriptInterface<IAssetFactoryInterface> GetAssetFactoryFromFactoryClass(TSubclassOf<UAssetFactoryInterface> InFactoryInterfaceClass) const;
+	TScriptInterface<IAssetFactoryInterface> GetAssetFactoryFromFactoryClass(TSubclassOf<UClass> InFactoryInterfaceClass) const;
 
 private:
 	void RegisterPlacementFactories();
diff --git a/Engine/Source/Editor/UnrealEd/Private/AssetSelection.cpp b/Engine/Source/Editor/UnrealEd/Private/AssetSelection.cpp
--- a/Engine/Source/Editor/UnrealEd/Private/AssetSelection.cpp
+++ b/Engine/Source/Editor/UnrealEd/Private/AssetSelection.cpp
@@ -616,7 +616,7 @@ static AActor* PrivateAddActor( UObject* Asset, UActorFactory* Factory, bool Sel
 			PlacementInfo.PreferredLevel = DesiredLevel;
 			PlacementInfo.NameOverride = Name;
 			PlacementInfo.FinalizedTransform = ActorTransform;
-			PlacementInfo.FactoryOverride = Factory->GetClass();
+			PlacementInfo.FactoryOverride = Factory;
 
 			FPlacementOptions PlacementOptions;
 			PlacementOptions.bIsCreatingPreviewElements = FLevelEditorViewportClient::IsDroppingPreviewActor();
