[Placement mode] Fix a crash when clearing the palette with single place tool active

diff --git a/Engine/Plugins/Experimental/AssetPlacementEdMode/Source/AssetPlacementEdMode/Private/Tools/PlacementPlaceSingleTool.cpp b/Engine/Plugins/Experimental/AssetPlacementEdMode/Source/AssetPlacementEdMode/Private/Tools/PlacementPlaceSingleTool.cpp
--- a/Engine/Plugins/Experimental/AssetPlacementEdMode/Source/AssetPlacementEdMode/Private/Tools/PlacementPlaceSingleTool.cpp
+++ b/Engine/Plugins/Experimental/AssetPlacementEdMode/Source/AssetPlacementEdMode/Private/Tools/PlacementPlaceSingleTool.cpp
@@ -268,6 +268,11 @@ void UPlacementModePlaceSingleTool::CreatePreviewElements(const FInputDeviceRay&
 		GeneratePlacementData(DevicePos);
 	}
 
+	if (!PlacementInfo)
+	{
+		return;
+	}
+
 	if (UPlacementSubsystem* PlacementSubsystem = GEditor->GetEditorSubsystem<UPlacementSubsystem>())
 	{
 		FPlacementOptions PlacementOptions;
