[Placement Mode] Fix an issue that prevented the user palette from being saved directly as an asset

diff --git a/Engine/Plugins/Experimental/AssetPlacementEdMode/Source/AssetPlacementEdMode/Private/AssetPlacementSettings.cpp b/Engine/Plugins/Experimental/AssetPlacementEdMode/Source/AssetPlacementEdMode/Private/AssetPlacementSettings.cpp
--- a/Engine/Plugins/Experimental/AssetPlacementEdMode/Source/AssetPlacementEdMode/Private/AssetPlacementSettings.cpp
+++ b/Engine/Plugins/Experimental/AssetPlacementEdMode/Source/AssetPlacementEdMode/Private/AssetPlacementSettings.cpp
@@ -13,6 +13,8 @@
 #include "Elements/Interfaces/TypedElementAssetDataInterface.h"
 #include "Instances/InstancedPlacementClientInfo.h"
 
+#include "AssetToolsModule.h"
+
 bool UAssetPlacementSettings::CanEditChange(const FProperty* InProperty) const
 {
 	if (!Super::CanEditChange(InProperty))
@@ -144,6 +146,13 @@ void UAssetPlacementSettings::SaveActivePalette()
 	{
 		UPackageTools::SavePackagesForObjects(TArray<UObject*>({ ActivePalette }));
 	}
+	else if (UserPalette)
+	{
+		FAssetToolsModule& AssetToolsModule = FAssetToolsModule::GetModule();
+		UObject* NewPaletteAsset = AssetToolsModule.Get().DuplicateAssetWithDialogAndTitle(FString(), FString(), UserPalette, NSLOCTEXT("AssetPlacementEdMode", "SavePaletteAsDialogTitle", "Save Asset Palette As..."));
+		ActivePalette = CastChecked<UPlacementPaletteAsset>(NewPaletteAsset);
+		LastActivePalettePath = FSoftObjectPath(ActivePalette);
+	}
 }
 
 bool UAssetPlacementSettings::DoesActivePaletteSupportElement(const FTypedElementHandle& InElementToCheck) const
diff --git a/Engine/Plugins/Experimental/AssetPlacementEdMode/Source/AssetPlacementEdMode/Private/SAssetPlacementPalette.cpp b/Engine/Plugins/Experimental/AssetPlacementEdMode/Source/AssetPlacementEdMode/Private/SAssetPlacementPalette.cpp
--- a/Engine/Plugins/Experimental/AssetPlacementEdMode/Source/AssetPlacementEdMode/Private/SAssetPlacementPalette.cpp
+++ b/Engine/Plugins/Experimental/AssetPlacementEdMode/Source/AssetPlacementEdMode/Private/SAssetPlacementPalette.cpp
@@ -130,7 +130,7 @@ void SAssetPlacementPalette::Construct(const FArguments& InArgs)
 							.AutoWidth()
 							.Padding(2.0f, 0.0f)
 							[
-								PropertyCustomizationHelpers::MakeSaveButton(FSimpleDelegate::CreateSP(this, &SAssetPlacementPalette::OnSavePaletteAssetClicked), LOCTEXT("SaveAssetPaletteTooltip", "Save changes to the current palette asset"), TAttribute<bool>::CreateLambda([this]() { return PalettePath.IsValid(); }))
+								PropertyCustomizationHelpers::MakeSaveButton(FSimpleDelegate::CreateSP(this, &SAssetPlacementPalette::OnSavePaletteAssetClicked), LOCTEXT("SaveAssetPaletteTooltip", "Save changes to the current palette asset"), TAttribute<bool>::CreateLambda([this]() { return PalettePath.IsValid() || (PaletteItems.Num() > 0); }))
 							]
 						]
 					]
@@ -403,6 +403,12 @@ void SAssetPlacementPalette::OnSavePaletteAssetClicked()
 	if (UPlacementModeSubsystem* PlacementModeSubsystem = GEditor->GetEditorSubsystem<UPlacementModeSubsystem>())
 	{
 		PlacementModeSubsystem->GetMutableModeSettingsObject()->SaveActivePalette();
+
+		// If we saved from a temporary/user palette, update the active palette now
+		if (!PalettePath.IsValid())
+		{
+			OnSetPaletteAsset(PlacementModeSubsystem->GetModeSettingsObject()->GetActivePalettePath().TryLoad());
+		}
 	}
 }
 
