[Placement mode] Change default setting for select after place in single place tool to false based on user feedback.

diff --git a/Engine/Plugins/Experimental/AssetPlacementEdMode/Source/AssetPlacementEdMode/Private/Tools/PlacementPlaceSingleTool.h b/Engine/Plugins/Experimental/AssetPlacementEdMode/Source/AssetPlacementEdMode/Private/Tools/PlacementPlaceSingleTool.h
--- a/Engine/Plugins/Experimental/AssetPlacementEdMode/Source/AssetPlacementEdMode/Private/Tools/PlacementPlaceSingleTool.h
+++ b/Engine/Plugins/Experimental/AssetPlacementEdMode/Source/AssetPlacementEdMode/Private/Tools/PlacementPlaceSingleTool.h
@@ -38,7 +38,7 @@ public:
 
 	/** If the tool should automatically select the last placed asset. */
 	UPROPERTY(config, EditAnywhere, Category = "Single Place Settings")
-	bool bSelectAfterPlacing = true;
+	bool bSelectAfterPlacing = false;
 };
 
 UCLASS(Transient)
