[Placement Mode] Add toggle to placement mode for instance workflows

#Add cvar to disable placement mode instance workflows, and disable entirely in cooked build

#rnx

diff --git a/Engine/Plugins/Experimental/AssetPlacementEdMode/Source/AssetPlacementEdMode/Private/AssetPlacementEdMode.cpp b/Engine/Plugins/Experimental/AssetPlacementEdMode/Source/AssetPlacementEdMode/Private/AssetPlacementEdMode.cpp
--- a/Engine/Plugins/Experimental/AssetPlacementEdMode/Source/AssetPlacementEdMode/Private/AssetPlacementEdMode.cpp
+++ b/Engine/Plugins/Experimental/AssetPlacementEdMode/Source/AssetPlacementEdMode/Private/AssetPlacementEdMode.cpp
@@ -10,6 +10,7 @@
 #include "EditorModeManager.h"
 #include "Selection.h"
 #include "EngineUtils.h"
+#include "AssetPlacementEdModeModule.h"
 
 #include "Tools/AssetEditorContextInterface.h"
 #include "Tools/PlacementSelectTool.h"
@@ -26,6 +27,7 @@
 
 #include "Settings/LevelEditorMiscSettings.h"
 #include "Modes/PlacementModeSubsystem.h"
+#include "AssetPlacementEdModeModule.h"
 
 #define LOCTEXT_NAMESPACE "AssetPlacementEdMode"
 
@@ -56,10 +58,16 @@ void UAssetPlacementEdMode::Enter()
 
 	const FAssetPlacementEdModeCommands& PlacementModeCommands = FAssetPlacementEdModeCommands::Get();
 	RegisterTool(PlacementModeCommands.Select, UPlacementModeSelectTool::ToolName, NewObject<UPlacementModeSelectToolBuilder>(this));
-	RegisterTool(PlacementModeCommands.Place, UPlacementModePlacementTool::ToolName, NewObject<UPlacementModePlacementToolBuilder>(this));
-	RegisterTool(PlacementModeCommands.LassoSelect, UPlacementModeLassoSelectTool::ToolName, NewObject<UPlacementModeLassoSelectToolBuilder>(this));
 	RegisterTool(PlacementModeCommands.PlaceSingle, UPlacementModePlaceSingleTool::ToolName, NewObject<UPlacementModePlaceSingleToolBuilder>(this));
-	RegisterTool(PlacementModeCommands.Erase, UPlacementModeEraseTool::ToolName, NewObject<UPlacementModeEraseToolBuilder>(this));
+
+#if !UE_IS_COOKED_EDITOR
+	if (AssetPlacementEdModeUtil::AreInstanceWorkflowsEnabled())
+	{
+		RegisterTool(PlacementModeCommands.Place, UPlacementModePlacementTool::ToolName, NewObject<UPlacementModePlacementToolBuilder>(this));
+		RegisterTool(PlacementModeCommands.LassoSelect, UPlacementModeLassoSelectTool::ToolName, NewObject<UPlacementModeLassoSelectToolBuilder>(this));
+		RegisterTool(PlacementModeCommands.Erase, UPlacementModeEraseTool::ToolName, NewObject<UPlacementModeEraseToolBuilder>(this));
+	}
+#endif // !UE_IS_COOKED_EDITOR
 
 	// Stash the current editor selection, since this mode will modify it.
 	Owner->StoreSelection(AssetPlacementEdModeID);
@@ -71,7 +79,12 @@ void UAssetPlacementEdMode::Enter()
 	// Enable the select tool by default.
 	GetInteractiveToolsContext()->StartTool(UPlacementModeSelectTool::ToolName);
 
-	SMInstanceElementDataUtil::OnSMInstanceElementsEnabledChanged().AddUObject(this, &UAssetPlacementEdMode::OnSMIsntancedElementsEnabledChanged);
+#if !UE_IS_COOKED_EDITOR
+	if (AssetPlacementEdModeUtil::AreInstanceWorkflowsEnabled())
+	{
+		SMInstanceElementDataUtil::OnSMInstanceElementsEnabledChanged().AddUObject(this, &UAssetPlacementEdMode::OnSMIsntancedElementsEnabledChanged);
+	}
+#endif // !UE_IS_COOKED_EDITOR
 }
 
 void UAssetPlacementEdMode::Exit()
diff --git a/Engine/Plugins/Experimental/AssetPlacementEdMode/Source/AssetPlacementEdMode/Private/AssetPlacementEdModeCommands.cpp b/Engine/Plugins/Experimental/AssetPlacementEdMode/Source/AssetPlacementEdMode/Private/AssetPlacementEdModeCommands.cpp
--- a/Engine/Plugins/Experimental/AssetPlacementEdMode/Source/AssetPlacementEdMode/Private/AssetPlacementEdModeCommands.cpp
+++ b/Engine/Plugins/Experimental/AssetPlacementEdMode/Source/AssetPlacementEdMode/Private/AssetPlacementEdModeCommands.cpp
@@ -2,6 +2,7 @@
 
 #include "AssetPlacementEdModeCommands.h"
 #include "AssetPlacementEdModeStyle.h"
+#include "AssetPlacementEdModeModule.h"
 
 #define LOCTEXT_NAMESPACE "AssetPlacementEdMode"
 
@@ -18,14 +19,27 @@ void FAssetPlacementEdModeCommands::RegisterCommands()
 	TArray <TSharedPtr<FUICommandInfo>>& ToolCommands = Commands.FindOrAdd(NAME_Default);
 	UI_COMMAND(Select, "Select", "Select by clicking single assets matching the active palette.", EUserInterfaceActionType::ToggleButton, FInputChord());
 	ToolCommands.Add(Select);
-	UI_COMMAND(LassoSelect, "Lasso", "Selects asset by painting the area to select.", EUserInterfaceActionType::ToggleButton, FInputChord());
-	ToolCommands.Add(LassoSelect);
-	UI_COMMAND(Place, "Paint", "Paint mutliple assets from the active palette.", EUserInterfaceActionType::ToggleButton, FInputChord());
-	ToolCommands.Add(Place);
+
+#if !UE_IS_COOKED_EDITOR
+	if (AssetPlacementEdModeUtil::AreInstanceWorkflowsEnabled())
+	{
+		UI_COMMAND(LassoSelect, "Lasso", "Selects asset by painting the area to select.", EUserInterfaceActionType::ToggleButton, FInputChord());
+		ToolCommands.Add(LassoSelect);
+		UI_COMMAND(Place, "Paint", "Paint mutliple assets from the active palette.", EUserInterfaceActionType::ToggleButton, FInputChord());
+		ToolCommands.Add(Place);
+	}
+#endif // !UE_IS_COOKED_EDITOR
+
 	UI_COMMAND(PlaceSingle, "Single", "Place a single, random asset from the active palette.", EUserInterfaceActionType::ToggleButton, FInputChord());
 	ToolCommands.Add(PlaceSingle);
-	UI_COMMAND(Erase, "Erase", "Paint to erase assets matching the active palette.", EUserInterfaceActionType::ToggleButton, FInputChord());
-	ToolCommands.Add(Erase);
+
+#if !UE_IS_COOKED_EDITOR
+	if (AssetPlacementEdModeUtil::AreInstanceWorkflowsEnabled())
+	{
+		UI_COMMAND(Erase, "Erase", "Paint to erase assets matching the active palette.", EUserInterfaceActionType::ToggleButton, FInputChord());
+		ToolCommands.Add(Erase);
+	}
+#endif // !UE_IS_COOKED_EDITOR
 }
 
 TMap<FName, TArray<TSharedPtr<FUICommandInfo>>> FAssetPlacementEdModeCommands::GetCommands()
diff --git a/Engine/Plugins/Experimental/AssetPlacementEdMode/Source/AssetPlacementEdMode/Private/AssetPlacementEdModeModule.cpp b/Engine/Plugins/Experimental/AssetPlacementEdMode/Source/AssetPlacementEdMode/Private/AssetPlacementEdModeModule.cpp
--- a/Engine/Plugins/Experimental/AssetPlacementEdMode/Source/AssetPlacementEdMode/Private/AssetPlacementEdModeModule.cpp
+++ b/Engine/Plugins/Experimental/AssetPlacementEdMode/Source/AssetPlacementEdMode/Private/AssetPlacementEdModeModule.cpp
@@ -11,6 +11,32 @@
 
 IMPLEMENT_MODULE( FAssetPlacementEdMode, AssetPlacementEdMode );
 
+#if !UE_IS_COOKED_EDITOR
+namespace AssetPlacementEdModeInternal
+{
+	static int32 GEnableInstanceWorkflows = 1;
+	static FAutoConsoleVariableRef CVarEnableInstanceWorkflows(
+		TEXT("PlacementMode.EnableInstanceWorkflows"),
+		GEnableInstanceWorkflows,
+		TEXT("Is support for instances everywhere in placement mode enabled?"),
+		FConsoleVariableDelegate::CreateLambda([](IConsoleVariable* InVariable)
+		{
+			// Force a reset on the registered commands to rebuild the modes' set of valid toolbox commands
+			FAssetPlacementEdModeCommands::Unregister();
+			FAssetPlacementEdModeCommands::Register();
+		})
+	);
+}
+
+namespace AssetPlacementEdModeUtil
+{
+	bool AreInstanceWorkflowsEnabled()
+	{
+		return AssetPlacementEdModeInternal::GEnableInstanceWorkflows > 0;
+	}
+}
+#endif // !UE_IS_COOKED_EDITOR
+
 void FAssetPlacementEdMode::StartupModule()
 {
 	FAssetPlacementEdModeStyle::Get();
diff --git a/Engine/Plugins/Experimental/AssetPlacementEdMode/Source/AssetPlacementEdMode/Private/Tools/PlacementPlaceSingleTool.cpp b/Engine/Plugins/Experimental/AssetPlacementEdMode/Source/AssetPlacementEdMode/Private/Tools/PlacementPlaceSingleTool.cpp
--- a/Engine/Plugins/Experimental/AssetPlacementEdMode/Source/AssetPlacementEdMode/Private/Tools/PlacementPlaceSingleTool.cpp
+++ b/Engine/Plugins/Experimental/AssetPlacementEdMode/Source/AssetPlacementEdMode/Private/Tools/PlacementPlaceSingleTool.cpp
@@ -25,6 +25,10 @@
 #include "ContextObjectStore.h"
 #include "PlacementPaletteItem.h"
 
+#if !UE_IS_COOKED_EDITOR
+#include "AssetPlacementEdModeModule.h"
+#endif
+
 constexpr TCHAR UPlacementModePlaceSingleTool::ToolName[];
 
 UPlacementBrushToolBase* UPlacementModePlaceSingleToolBuilder::FactoryToolInstance(UObject* Outer) const
@@ -86,7 +90,17 @@ void UPlacementModePlaceSingleTool::OnClickPress(const FInputDeviceRay& PressPos
 		// Place the Preview data if we managed to get to a valid handled click.
 		FPlacementOptions PlacementOptions;
 		PlacementOptions.bPreferBatchPlacement = true;
-		PlacementOptions.InstancedPlacementGridGuid = PlacementSettings->GetActivePaletteGuid();
+		
+#if !UE_IS_COOKED_EDITOR
+		if (!AssetPlacementEdModeUtil::AreInstanceWorkflowsEnabled())
+		{
+			PlacementOptions.InstancedPlacementGridGuid = FGuid();
+		}
+		else
+#endif
+		{
+			PlacementOptions.InstancedPlacementGridGuid = PlacementSettings->GetActivePaletteGuid();
+		}
 
 		FAssetPlacementInfo FinalizedPlacementInfo = *PlacementInfo;
 		FinalizedPlacementInfo.FinalizedTransform = FinalizeTransform(
diff --git a/Engine/Plugins/Experimental/AssetPlacementEdMode/Source/AssetPlacementEdMode/Public/AssetPlacementEdModeModule.h b/Engine/Plugins/Experimental/AssetPlacementEdMode/Source/AssetPlacementEdMode/Public/AssetPlacementEdModeModule.h
--- a/Engine/Plugins/Experimental/AssetPlacementEdMode/Source/AssetPlacementEdMode/Public/AssetPlacementEdModeModule.h
+++ b/Engine/Plugins/Experimental/AssetPlacementEdMode/Source/AssetPlacementEdMode/Public/AssetPlacementEdModeModule.h
@@ -7,6 +7,13 @@
 
 class IAssetTypeActions;
 
+#if !UE_IS_COOKED_EDITOR
+namespace AssetPlacementEdModeUtil
+{
+	bool AreInstanceWorkflowsEnabled();
+}
+#endif // !UE_IS_COOKED_EDITOR
+
 class FAssetPlacementEdMode : public IModuleInterface
 {
 public:
