[Placement Mode] Disable custom palette save/load and fix an issue where mirror content browser was not applied on first run

#rnx

diff --git a/Engine/Plugins/Experimental/AssetPlacementEdMode/Source/AssetPlacementEdMode/Private/SAssetPlacementPalette.cpp b/Engine/Plugins/Experimental/AssetPlacementEdMode/Source/AssetPlacementEdMode/Private/SAssetPlacementPalette.cpp
--- a/Engine/Plugins/Experimental/AssetPlacementEdMode/Source/AssetPlacementEdMode/Private/SAssetPlacementPalette.cpp
+++ b/Engine/Plugins/Experimental/AssetPlacementEdMode/Source/AssetPlacementEdMode/Private/SAssetPlacementPalette.cpp
@@ -41,6 +41,10 @@
 #include "Factories/AssetFactoryInterface.h"
 #include "Instances/EditorPlacementSettings.h"
 
+#if !UE_IS_COOKED_EDITOR
+#include "AssetPlacementEdModeModule.h"
+#endif
+
 #define LOCTEXT_NAMESPACE "AssetPlacementMode"
 
 ////////////////////////////////////////////////
@@ -72,6 +76,7 @@ void SAssetPlacementPalette::Construct(const FArguments& InArgs)
 			OnSetPaletteAsset(PlacementSettings->GetActivePalettePath().TryLoad());
 		}
 	}
+	SetupContentBrowserMirroring(!bIsMirroringContentBrowser);
 
 	ChildSlot
 	[
@@ -112,6 +117,7 @@ void SAssetPlacementPalette::Construct(const FArguments& InArgs)
 					.ObjectPath(this, &SAssetPlacementPalette::GetPalettePath)
 					.OnObjectChanged(this, &SAssetPlacementPalette::OnSetPaletteAsset)
 					.ThumbnailPool(ThumbnailPool)
+					.Visibility(this, &SAssetPlacementPalette::GetPaletteAssetPropertyBoxVisible)
 					.CustomContentSlot()
 					[
 						SNew(SBox)
@@ -417,6 +423,18 @@ EVisibility SAssetPlacementPalette::GetContentBrowserMirrorVisibility() const
 	return GetPalettePath().IsEmpty() ? EVisibility::Visible : EVisibility::Collapsed;
 }
 
+EVisibility SAssetPlacementPalette::GetPaletteAssetPropertyBoxVisible() const
+{
+#if !UE_IS_COOKED_EDITOR
+	if (AssetPlacementEdModeUtil::AreInstanceWorkflowsEnabled())
+	{
+		return EVisibility::Visible;
+	}
+#endif
+	
+	return EVisibility::Collapsed;
+}
+
 void SAssetPlacementPalette::OnContentBrowserMirrorButtonClicked(ECheckBoxState InState)
 {
 	SetupContentBrowserMirroring((InState == ECheckBoxState::Checked));
@@ -438,12 +456,17 @@ void SAssetPlacementPalette::OnContentBrowserSelectionChanged(const TArray<FAsse
 
 void SAssetPlacementPalette::SetupContentBrowserMirroring(bool bInMirrorContentBrowser)
 {
-	bool bWasMirroringContentBrowser = bIsMirroringContentBrowser;
-	if (bWasMirroringContentBrowser != bInMirrorContentBrowser)
+	bool bCanMirrorContentBrowser = bInMirrorContentBrowser;
+	if (PalettePath.IsValid())
+	{
+		bCanMirrorContentBrowser = false;
+	}
+
+	if (bIsMirroringContentBrowser != bCanMirrorContentBrowser)
 	{
 		if (FContentBrowserModule* ContentBrowserModule = FModuleManager::GetModulePtr<FContentBrowserModule>("ContentBrowser"))
 		{
-			if (bInMirrorContentBrowser)
+			if (bInMirrorContentBrowser && bCanMirrorContentBrowser)
 			{
 				TArray<FAssetData> SelectedAssetDatas;
 				ContentBrowserModule->Get().GetSelectedAssets(SelectedAssetDatas);
@@ -475,10 +498,6 @@ void SAssetPlacementPalette::OnSetPaletteAsset(const FAssetData& InAssetData)
 			PlacementModeSubsystem->GetMutableModeSettingsObject()->SetPaletteAsset(PaletteAsset);
 
 			PalettePath = NewAssetPath;
-
-			// Setup content browser mirroring if this is the user palette
-			SetupContentBrowserMirroring(PalettePath.IsValid() ? false : PlacementModeSubsystem->GetModeSettingsObject()->bUseContentBrowserSelection);
-
 			SetPaletteItems(PlacementModeSubsystem->GetModeSettingsObject()->GetActivePaletteItems());
 		}
 	}
diff --git a/Engine/Plugins/Experimental/AssetPlacementEdMode/Source/AssetPlacementEdMode/Private/SAssetPlacementPalette.h b/Engine/Plugins/Experimental/AssetPlacementEdMode/Source/AssetPlacementEdMode/Private/SAssetPlacementPalette.h
--- a/Engine/Plugins/Experimental/AssetPlacementEdMode/Source/AssetPlacementEdMode/Private/SAssetPlacementPalette.h
+++ b/Engine/Plugins/Experimental/AssetPlacementEdMode/Source/AssetPlacementEdMode/Private/SAssetPlacementPalette.h
@@ -97,6 +97,7 @@ private:
 	void OnResetPaletteAssetClicked();
 	void OnSavePaletteAssetClicked();
 	EVisibility GetContentBrowserMirrorVisibility() const;
+	EVisibility GetPaletteAssetPropertyBoxVisible() const;
 	void OnContentBrowserMirrorButtonClicked(ECheckBoxState InState);
 	void OnContentBrowserSelectionChanged(const TArray<FAssetData>& NewSelectedAssets, bool bIsPrimaryBrowser);
 	void SetupContentBrowserMirroring(bool bInMirrorContentBrowser);
