[Placement Mode] Add toggle for details panel tab for per item settings

#rnx

diff --git a/Engine/Plugins/Experimental/AssetPlacementEdMode/Source/AssetPlacementEdMode/Private/AssetPlacementEdModeToolkit.cpp b/Engine/Plugins/Experimental/AssetPlacementEdMode/Source/AssetPlacementEdMode/Private/AssetPlacementEdModeToolkit.cpp
--- a/Engine/Plugins/Experimental/AssetPlacementEdMode/Source/AssetPlacementEdMode/Private/AssetPlacementEdModeToolkit.cpp
+++ b/Engine/Plugins/Experimental/AssetPlacementEdMode/Source/AssetPlacementEdMode/Private/AssetPlacementEdModeToolkit.cpp
@@ -11,6 +11,10 @@
 #include "PropertyEditorModule.h"
 #include "Modules/ModuleManager.h"
 
+#if !UE_IS_COOKED_EDITOR
+#include "AssetPlacementEdModeModule.h"
+#endif // !UE_IS_COOKED_EDITOR
+
 #define LOCTEXT_NAMESPACE "AssetPlacementEdModeToolkit"
 
 FAssetPlacementEdModeToolkit::FAssetPlacementEdModeToolkit()
@@ -68,11 +72,16 @@
 		AssetPaletteInfo.TabIcon = GetEditorModeIcon();
 		ModeUILayerPtr->SetModePanelInfo(UAssetEditorUISubsystem::BottomLeftTabID, AssetPaletteInfo);
 
-		PaletteItemDetailsViewInfo.OnSpawnTab = FOnSpawnTab::CreateSP(SharedThis(this), &FAssetPlacementEdModeToolkit::CreatePaletteItemDetailsView);
-		PaletteItemDetailsViewInfo.TabLabel = LOCTEXT("PaletteItemDetailsTab", "Palette Details");
-		PaletteItemDetailsViewInfo.TabTooltip = LOCTEXT("PaletteItemDetailsTabTooltipText", "Open the asset palette details tab, which allows customization of individual items in the active palette.");
-		PaletteItemDetailsViewInfo.TabIcon = GetEditorModeIcon();
-		ModeUILayerPtr->SetModePanelInfo(UAssetEditorUISubsystem::BottomRightTabID, PaletteItemDetailsViewInfo);
+#if !UE_IS_COOKED_EDITOR
+		if (AssetPlacementEdModeUtil::AreInstanceWorkflowsEnabled())
+		{
+			PaletteItemDetailsViewInfo.OnSpawnTab = FOnSpawnTab::CreateSP(SharedThis(this), &FAssetPlacementEdModeToolkit::CreatePaletteItemDetailsView);
+			PaletteItemDetailsViewInfo.TabLabel = LOCTEXT("PaletteItemDetailsTab", "Palette Details");
+			PaletteItemDetailsViewInfo.TabTooltip = LOCTEXT("PaletteItemDetailsTabTooltipText", "Open the asset palette details tab, which allows customization of individual items in the active palette.");
+			PaletteItemDetailsViewInfo.TabIcon = GetEditorModeIcon();
+			ModeUILayerPtr->SetModePanelInfo(UAssetEditorUISubsystem::BottomRightTabID, PaletteItemDetailsViewInfo);
+		}
+#endif // !UE_IS_COOKED_EDITOR
 	}
 }
 
