[Placement Mode] Only filter selection based on the active tool, not on the palette

#rnx

diff --git a/Engine/Plugins/Experimental/AssetPlacementEdMode/Source/AssetPlacementEdMode/Private/AssetPlacementEdMode.cpp b/Engine/Plugins/Experimental/AssetPlacementEdMode/Source/AssetPlacementEdMode/Private/AssetPlacementEdMode.cpp
--- a/Engine/Plugins/Experimental/AssetPlacementEdMode/Source/AssetPlacementEdMode/Private/AssetPlacementEdMode.cpp
+++ b/Engine/Plugins/Experimental/AssetPlacementEdMode/Source/AssetPlacementEdMode/Private/AssetPlacementEdMode.cpp
@@ -131,14 +131,7 @@ bool UAssetPlacementEdMode::IsSelectionAllowed(AActor* InActor, bool bInSelectio
 	}
 
 	// Otherwise, need to be in selection tool for selection to be allowed
-	if (!bIsInSelectionTool)
-	{
-		return false;
-	}
-
-	// And we need to have a valid palette item.
-	FTypedElementHandle ActorHandle = UEngineElementsLibrary::AcquireEditorActorElementHandle(InActor);
-	return SettingsObjectAsPlacementSettings->DoesActivePaletteSupportElement(ActorHandle);
+	return bIsInSelectionTool;
 }
 
 void UAssetPlacementEdMode::OnToolStarted(UInteractiveToolManager* Manager, UInteractiveTool* Tool)
