Fix incorrect attempt at finding NormalText style.

diff --git a/Engine/Plugins/Experimental/AssetPlacementEdMode/Source/AssetPlacementEdMode/Private/AssetPlacementEdModeStyle.cpp b/Engine/Plugins/Experimental/AssetPlacementEdMode/Source/AssetPlacementEdMode/Private/AssetPlacementEdModeStyle.cpp
--- a/Engine/Plugins/Experimental/AssetPlacementEdMode/Source/AssetPlacementEdMode/Private/AssetPlacementEdModeStyle.cpp
+++ b/Engine/Plugins/Experimental/AssetPlacementEdMode/Source/AssetPlacementEdMode/Private/AssetPlacementEdModeStyle.cpp
@@ -70,7 +70,7 @@ FAssetPlacementEdModeStyle::FAssetPlacementEdModeStyle()
 	Set("AssetPlacementEdMode.MoveToActivePartition",       new IMAGE_BRUSH("Icons/GeneralTools/MoveToLevel_40x", Icon20x20));
 	Set("AssetPlacementEdMode.MoveToActivePartition.Small", new IMAGE_BRUSH("Icons/GeneralTools/MoveToLevel_40x", Icon20x20));
 
-	Set("AssetPlacementEdMode.AddAssetType.Text", FTextBlockStyle(GetWidgetStyle<FTextBlockStyle>("NormalText"))
+	Set("AssetPlacementEdMode.AddAssetType.Text", FTextBlockStyle(FAppStyle::Get().GetWidgetStyle<FTextBlockStyle>("NormalText"))
 		.SetFont(DEFAULT_FONT("Bold", 10))
 		.SetColorAndOpacity(FLinearColor(1.0f, 1.0f, 1.0f))
 		.SetHighlightColor(FLinearColor(1.0f, 1.0f, 1.0f))
