[Placement Mode] Fix backwards logic which incorrectly caused ISM placement in cooked editors

diff --git a/Engine/Plugins/Experimental/AssetPlacementEdMode/Source/AssetPlacementEdMode/Private/Tools/PlacementPlaceSingleTool.cpp b/Engine/Plugins/Experimental/AssetPlacementEdMode/Source/AssetPlacementEdMode/Private/Tools/PlacementPlaceSingleTool.cpp
--- a/Engine/Plugins/Experimental/AssetPlacementEdMode/Source/AssetPlacementEdMode/Private/Tools/PlacementPlaceSingleTool.cpp
+++ b/Engine/Plugins/Experimental/AssetPlacementEdMode/Source/AssetPlacementEdMode/Private/Tools/PlacementPlaceSingleTool.cpp
@@ -92,14 +92,14 @@ void UPlacementModePlaceSingleTool::OnClickPress(const FInputDeviceRay& PressPos
 		PlacementOptions.bPreferBatchPlacement = true;
 		
 #if !UE_IS_COOKED_EDITOR
-		if (!AssetPlacementEdModeUtil::AreInstanceWorkflowsEnabled())
+		if (AssetPlacementEdModeUtil::AreInstanceWorkflowsEnabled())
 		{
-			PlacementOptions.InstancedPlacementGridGuid = FGuid();
+			PlacementOptions.InstancedPlacementGridGuid = PlacementSettings->GetActivePaletteGuid();
 		}
 		else
 #endif
 		{
-			PlacementOptions.InstancedPlacementGridGuid = PlacementSettings->GetActivePaletteGuid();
+			PlacementOptions.InstancedPlacementGridGuid = FGuid();
 		}
 
 		FAssetPlacementInfo FinalizedPlacementInfo = *PlacementInfo;
