[Placement Mode] Make the valid drop area in the palette match remaining available size in the dock tab

diff --git a/Engine/Plugins/Experimental/AssetPlacementEdMode/Source/AssetPlacementEdMode/Private/AssetPlacementEdModeToolkit.cpp b/Engine/Plugins/Experimental/AssetPlacementEdMode/Source/AssetPlacementEdMode/Private/AssetPlacementEdModeToolkit.cpp
--- a/Engine/Plugins/Experimental/AssetPlacementEdMode/Source/AssetPlacementEdMode/Private/AssetPlacementEdModeToolkit.cpp
+++ b/Engine/Plugins/Experimental/AssetPlacementEdMode/Source/AssetPlacementEdMode/Private/AssetPlacementEdModeToolkit.cpp
@@ -91,12 +91,8 @@
 {
 	return SNew(SDockTab)
 		[
-			SNew(SScrollBox)
-			+ SScrollBox::Slot()
-			[
-				SNew(SAssetPlacementPalette)
-				.ItemDetailsView(PaletteItemDetailsView)
-			]
+			SNew(SAssetPlacementPalette)
+			.ItemDetailsView(PaletteItemDetailsView)
 		];
 }
 
diff --git a/Engine/Plugins/Experimental/AssetPlacementEdMode/Source/AssetPlacementEdMode/Private/SAssetPlacementPalette.cpp b/Engine/Plugins/Experimental/AssetPlacementEdMode/Source/AssetPlacementEdMode/Private/SAssetPlacementPalette.cpp
--- a/Engine/Plugins/Experimental/AssetPlacementEdMode/Source/AssetPlacementEdMode/Private/SAssetPlacementPalette.cpp
+++ b/Engine/Plugins/Experimental/AssetPlacementEdMode/Source/AssetPlacementEdMode/Private/SAssetPlacementPalette.cpp
@@ -78,109 +78,109 @@ void SAssetPlacementPalette::Construct(const FArguments& InArgs)
 	}
 	SetupContentBrowserMirroring(!bIsMirroringContentBrowser);
 
-	ChildSlot
-	[
-		SNew(SVerticalBox)
-		+ SVerticalBox::Slot()
-		.HAlign(HAlign_Fill)
-		.AutoHeight()
+	TSharedPtr<SWidget> TopBar = 
+		SNew(SBorder)
+		.BorderImage(FAppStyle::Get().GetBrush("DetailsView.CategoryTop"))
+		.BorderBackgroundColor(FLinearColor(.6f, .6f, .6f, 1.0f))
 		[
-			// Top bar
-			SNew(SBorder)
-			.BorderImage(FAppStyle::Get().GetBrush("DetailsView.CategoryTop"))
-			.BorderBackgroundColor(FLinearColor(.6f, .6f, .6f, 1.0f))
+			SNew(SVerticalBox)
+			+ SVerticalBox::Slot()
+			.AutoHeight()
+			.HAlign(HAlign_Fill)
 			[
-				SNew(SVerticalBox)
-				+ SVerticalBox::Slot()
-				.AutoHeight()
-				.HAlign(HAlign_Fill)
+				SNew(SCheckBox)
+				.Type(ESlateCheckBoxType::Type::ToggleButton)
+				.IsChecked(bIsMirroringContentBrowser ? ECheckBoxState::Checked : ECheckBoxState::Unchecked)
+				.Style(&FAssetPlacementEdModeStyle::Get().GetWidgetStyle<FCheckBoxStyle>("ToggleButtonCheckBox"))
+				.Visibility(this, &SAssetPlacementPalette::GetContentBrowserMirrorVisibility)
+				.OnCheckStateChanged(this, &SAssetPlacementPalette::OnContentBrowserMirrorButtonClicked)
 				[
-					SNew(SCheckBox)
-					.Type(ESlateCheckBoxType::Type::ToggleButton)
-					.IsChecked(bIsMirroringContentBrowser ? ECheckBoxState::Checked : ECheckBoxState::Unchecked)
-					.Style(&FAssetPlacementEdModeStyle::Get().GetWidgetStyle<FCheckBoxStyle>("ToggleButtonCheckBox"))
-					.Visibility(this, &SAssetPlacementPalette::GetContentBrowserMirrorVisibility)
-					.OnCheckStateChanged(this, &SAssetPlacementPalette::OnContentBrowserMirrorButtonClicked)
-					[
-						SNew(STextBlock)
-						.Text(LOCTEXT("Placement_ToggleContentBrowserMirroring", "Mirror Content Browser Selection"))
-						.Justification(ETextJustify::Type::Center)
-						.TextStyle(&FAssetPlacementEdModeStyle::Get().GetWidgetStyle<FTextBlockStyle>("ButtonText"))
-						.ToolTipText(LOCTEXT("Placement_ToggleContentBrowserMirroring_ToolTip", "Toggles palette to mirror the active content browser selection."))
-					]
+					SNew(STextBlock)
+					.Text(LOCTEXT("Placement_ToggleContentBrowserMirroring", "Mirror Content Browser Selection"))
+					.Justification(ETextJustify::Type::Center)
+					.TextStyle(&FAssetPlacementEdModeStyle::Get().GetWidgetStyle<FTextBlockStyle>("ButtonText"))
+					.ToolTipText(LOCTEXT("Placement_ToggleContentBrowserMirroring_ToolTip", "Toggles palette to mirror the active content browser selection."))
 				]
-				+SVerticalBox::Slot()
-				.AutoHeight()
+			]
+			+SVerticalBox::Slot()
+			.AutoHeight()
+			[
+				SNew(SObjectPropertyEntryBox)
+				.AllowedClass(UPlacementPaletteAsset::StaticClass())
+				.ObjectPath(this, &SAssetPlacementPalette::GetPalettePath)
+				.OnObjectChanged(this, &SAssetPlacementPalette::OnSetPaletteAsset)
+				.ThumbnailPool(ThumbnailPool)
+				.Visibility(this, &SAssetPlacementPalette::GetPaletteAssetPropertyBoxVisible)
+				.CustomContentSlot()
 				[
-					SNew(SObjectPropertyEntryBox)
-					.AllowedClass(UPlacementPaletteAsset::StaticClass())
-					.ObjectPath(this, &SAssetPlacementPalette::GetPalettePath)
-					.OnObjectChanged(this, &SAssetPlacementPalette::OnSetPaletteAsset)
-					.ThumbnailPool(ThumbnailPool)
-					.Visibility(this, &SAssetPlacementPalette::GetPaletteAssetPropertyBoxVisible)
-					.CustomContentSlot()
+					SNew(SBox)
+					.HAlign(HAlign_Left)
+					.VAlign(VAlign_Center)
 					[
-						SNew(SBox)
-						.HAlign(HAlign_Left)
+						SNew(SHorizontalBox)
+						+ SHorizontalBox::Slot()
+						.VAlign(VAlign_Center)
+						.AutoWidth()
+						[
+							PropertyCustomizationHelpers::MakeClearButton(FSimpleDelegate::CreateSP(this, &SAssetPlacementPalette::OnResetPaletteAssetClicked), LOCTEXT("ResetPaletteAssetTooltip", "Clear the current palette asset, and use the user's local palette."), TAttribute<bool>::CreateLambda([this]() { return PalettePath.IsValid(); }))
+						]
+						+ SHorizontalBox::Slot()
 						.VAlign(VAlign_Center)
+						.AutoWidth()
+						.Padding(2.0f, 0.0f)
 						[
-							SNew(SHorizontalBox)
-							+ SHorizontalBox::Slot()
-							.VAlign(VAlign_Center)
-							.AutoWidth()
-							[
-								PropertyCustomizationHelpers::MakeClearButton(FSimpleDelegate::CreateSP(this, &SAssetPlacementPalette::OnResetPaletteAssetClicked), LOCTEXT("ResetPaletteAssetTooltip", "Clear the current palette asset, and use the user's local palette."), TAttribute<bool>::CreateLambda([this]() { return PalettePath.IsValid(); }))
-							]
-							+ SHorizontalBox::Slot()
-							.VAlign(VAlign_Center)
-							.AutoWidth()
-							.Padding(2.0f, 0.0f)
-							[
-								PropertyCustomizationHelpers::MakeSaveButton(FSimpleDelegate::CreateSP(this, &SAssetPlacementPalette::OnSavePaletteAssetClicked), LOCTEXT("SaveAssetPaletteTooltip", "Save changes to the current palette asset"), TAttribute<bool>::CreateLambda([this]() { return PalettePath.IsValid() || (PaletteItems.Num() > 0); }))
-							]
+							PropertyCustomizationHelpers::MakeSaveButton(FSimpleDelegate::CreateSP(this, &SAssetPlacementPalette::OnSavePaletteAssetClicked), LOCTEXT("SaveAssetPaletteTooltip", "Save changes to the current palette asset"), TAttribute<bool>::CreateLambda([this]() { return PalettePath.IsValid() || (PaletteItems.Num() > 0); }))
 						]
 					]
 				]
+			]
 
-				+ SVerticalBox::Slot()
-				.VAlign(VAlign_Center)
-				.AutoHeight()
+			+ SVerticalBox::Slot()
+			.VAlign(VAlign_Center)
+			.AutoHeight()
+			[
+				SNew(SHorizontalBox)
+				+ SHorizontalBox::Slot()
+				.FillWidth(.75f)
 				[
-					SNew(SHorizontalBox)
-
-					+ SHorizontalBox::Slot()
-					.FillWidth(.75f)
-					[
-						SAssignNew(SearchBoxPtr, SSearchBox)
-						.HintText(LOCTEXT("SearchPlacementPaletteHint", "Search Palette"))
-						.OnTextChanged(this, &SAssetPlacementPalette::OnSearchTextChanged)
-					]
+					SAssignNew(SearchBoxPtr, SSearchBox)
+					.HintText(LOCTEXT("SearchPlacementPaletteHint", "Search Palette"))
+					.OnTextChanged(this, &SAssetPlacementPalette::OnSearchTextChanged)
+				]
 
-					// View Options
-					+ SHorizontalBox::Slot()
-					.AutoWidth()
+				// View Options
+				+ SHorizontalBox::Slot()
+				.AutoWidth()
+				[
+					SNew(SComboButton)
+					.ForegroundColor(FSlateColor::UseForeground())
+					.ButtonStyle(FAppStyle::Get(), "ToggleButton")
+					.OnGetMenuContent(this, &SAssetPlacementPalette::GetViewOptionsMenuContent)
+					.ButtonContent()
 					[
-						SNew(SComboButton)
-						.ForegroundColor(FSlateColor::UseForeground())
-						.ButtonStyle(FAppStyle::Get(), "ToggleButton")
-						.OnGetMenuContent(this, &SAssetPlacementPalette::GetViewOptionsMenuContent)
-						.ButtonContent()
+						SNew(SBox)
+						.VAlign(VAlign_Center)
 						[
-							SNew(SBox)
-							.VAlign(VAlign_Center)
-							[
-								SNew(SImage)
-								.Image(FAppStyle::Get().GetBrush("GenericViewButton"))
-							]
+							SNew(SImage)
+							.Image(FAppStyle::Get().GetBrush("GenericViewButton"))
 						]
 					]
 				]
 			]
-		]
+		];
 
+	ChildSlot
+	[
+		SNew(SVerticalBox)
 		+ SVerticalBox::Slot()
 		.HAlign(HAlign_Fill)
-		.VAlign(VAlign_Fill)
+		.AutoHeight()
+		[
+			TopBar.ToSharedRef()
+		]
+
+		+ SVerticalBox::Slot()
+		.FillHeight(1.0)
 		[
 			SNew(SAssetDropTarget)
 			.OnAreAssetsAcceptableForDrop(this, &SAssetPlacementPalette::OnAreAssetsValidForDrop)
@@ -190,7 +190,24 @@ void SAssetPlacementPalette::Construct(const FArguments& InArgs)
 				SNew(SOverlay)
 				+ SOverlay::Slot()
 				[
-					CreatePaletteViews()
+					SNew(SBox)
+					.HeightOverride_Lambda([TopBarWeakPtr = TWeakPtr<SWidget>(TopBar)]()
+					{
+						if (TSharedPtr<SWidget> ParentWidgetPin = TopBarWeakPtr.Pin())
+						{
+							return FOptionalSize(ParentWidgetPin->GetTickSpaceGeometry().GetLocalSize().Y - 1);
+						}
+						return FOptionalSize();
+					})
+				]
+
+				+ SOverlay::Slot()
+				[
+					SNew(SScrollBox)
+					+ SScrollBox::Slot()
+					[
+						CreatePaletteViews()
+					]
 				]
 
 				+ SOverlay::Slot()
