[Placement Mode] Add ability to the palette's right click context menu to selected items from the active palette

diff --git a/Engine/Plugins/Experimental/AssetPlacementEdMode/Source/AssetPlacementEdMode/Private/AssetPlacementSettings.cpp b/Engine/Plugins/Experimental/AssetPlacementEdMode/Source/AssetPlacementEdMode/Private/AssetPlacementSettings.cpp
--- a/Engine/Plugins/Experimental/AssetPlacementEdMode/Source/AssetPlacementEdMode/Private/AssetPlacementSettings.cpp
+++ b/Engine/Plugins/Experimental/AssetPlacementEdMode/Source/AssetPlacementEdMode/Private/AssetPlacementSettings.cpp
@@ -95,6 +95,11 @@ UPlacementPaletteClient* UAssetPlacementSettings::AddClientToActivePalette(const
 	return NewPaletteItem;
 }
 
+int32 UAssetPlacementSettings::RemoveClientFromActivePalette(const FAssetData& InAssetData)
+{
+	return GetMutableActivePalette()->PaletteItems.RemoveAll([InAssetData](const UPlacementPaletteClient* ItemIter) { return ItemIter ? (ItemIter->AssetPath == InAssetData.ToSoftObjectPath()) : false; });
+}
+
 TArrayView<const TObjectPtr<UPlacementPaletteClient>> UAssetPlacementSettings::GetActivePaletteItems() const
 {
 	return MakeArrayView(GetActivePalette()->PaletteItems);
diff --git a/Engine/Plugins/Experimental/AssetPlacementEdMode/Source/AssetPlacementEdMode/Private/AssetPlacementSettings.h b/Engine/Plugins/Experimental/AssetPlacementEdMode/Source/AssetPlacementEdMode/Private/AssetPlacementSettings.h
--- a/Engine/Plugins/Experimental/AssetPlacementEdMode/Source/AssetPlacementEdMode/Private/AssetPlacementSettings.h
+++ b/Engine/Plugins/Experimental/AssetPlacementEdMode/Source/AssetPlacementEdMode/Private/AssetPlacementSettings.h
@@ -124,6 +124,7 @@ public:
 
 	void SetPaletteAsset(UPlacementPaletteAsset* InPaletteAsset);
 	UPlacementPaletteClient* AddClientToActivePalette(const FAssetData& InAssetData);
+	int32 RemoveClientFromActivePalette(const FAssetData& InAssetData);
 	TArrayView<const TObjectPtr<UPlacementPaletteClient>> GetActivePaletteItems() const;
 	const FGuid GetActivePaletteGuid() const;
 	void ClearActivePaletteItems();
diff --git a/Engine/Plugins/Experimental/AssetPlacementEdMode/Source/AssetPlacementEdMode/Private/SAssetPlacementPalette.cpp b/Engine/Plugins/Experimental/AssetPlacementEdMode/Source/AssetPlacementEdMode/Private/SAssetPlacementPalette.cpp
--- a/Engine/Plugins/Experimental/AssetPlacementEdMode/Source/AssetPlacementEdMode/Private/SAssetPlacementPalette.cpp
+++ b/Engine/Plugins/Experimental/AssetPlacementEdMode/Source/AssetPlacementEdMode/Private/SAssetPlacementPalette.cpp
@@ -322,6 +322,41 @@ void SAssetPlacementPalette::SetPaletteToAssetDataList(TArrayView<const FAssetDa
 	}
 }
 
+void SAssetPlacementPalette::OnRemoveSelectedItemsFromPalette()
+{
+	UPlacementModeSubsystem* ModeSubsystem = GEditor->GetEditorSubsystem<UPlacementModeSubsystem>();
+	TSharedPtr<IDetailsView> PinnedItemDetails = ItemDetailsWidget.Pin();
+	TArray<TWeakObjectPtr<UObject>> SelectedObjects;
+	if (PinnedItemDetails)
+	{
+		SelectedObjects = PinnedItemDetails->GetSelectedObjects();
+	}
+
+	for (FPlacementPaletteItemModelPtr& PaletteItem : GetActiveViewWidget()->GetSelectedItems())
+	{
+		const FAssetPlacementUIInfoPtr PaletteInfo = PaletteItem->GetTypeUIInfo();
+		
+		if (ModeSubsystem)
+		{
+			ModeSubsystem->GetMutableModeSettingsObject()->RemoveClientFromActivePalette(PaletteInfo->AssetData);
+		}
+
+		SelectedObjects.Remove(PaletteInfo->SettingsObject);
+	}
+
+	// Update the palette's view to the updated palette
+	if (ModeSubsystem)
+	{
+		SetPaletteItems(ModeSubsystem->GetModeSettingsObject()->GetActivePaletteItems());
+	}
+
+	// Reset the item details to valid objects
+	if (PinnedItemDetails)
+	{
+		PinnedItemDetails->SetObjects(SelectedObjects);
+	}
+}
+
 TSharedRef<SWidgetSwitcher> SAssetPlacementPalette::CreatePaletteViews()
 {
 	const FText BlankText = FText::GetEmpty();
@@ -743,6 +778,18 @@ TSharedPtr<SWidget> SAssetPlacementPalette::ConstructPlacementTypeContextMenu()
 				FExecuteAction::CreateSP(this, &SAssetPlacementPalette::OnClearPalette),
 				FCanExecuteAction::CreateSP(this, &SAssetPlacementPalette::HasAnyItemInPalette))
 		);
+
+		const int32 NumSelectedItems = SAssetPlacementPalette::GetActiveViewWidget()->GetNumItemsSelected();
+		MenuBuilder.AddMenuEntry(
+			FText::Format(LOCTEXT("Palette_RemoveItem", "Remove {0} {0}|plural(one=item,other=items)"), FText::AsNumber(NumSelectedItems)),
+			FText::Format(LOCTEXT("Palette_RemoveItemDesc", "Removes the {0} selected {0}|plural(one=item,other=items) from the palette."), FText::AsNumber(NumSelectedItems)),
+			FSlateIcon(),
+			FUIAction(
+				FExecuteAction::CreateSP(this, &SAssetPlacementPalette::OnRemoveSelectedItemsFromPalette),
+				FCanExecuteAction::CreateSP(this, &SAssetPlacementPalette::HasAnyItemInPalette),
+				FIsActionChecked::CreateLambda([]() { return false; }),
+				FIsActionButtonVisible::CreateLambda([NumSelectedItems]() { return (NumSelectedItems > 0); }))
+		);
 	}
 	return MenuBuilder.MakeWidget();
 }
diff --git a/Engine/Plugins/Experimental/AssetPlacementEdMode/Source/AssetPlacementEdMode/Private/SAssetPlacementPalette.h b/Engine/Plugins/Experimental/AssetPlacementEdMode/Source/AssetPlacementEdMode/Private/SAssetPlacementPalette.h
--- a/Engine/Plugins/Experimental/AssetPlacementEdMode/Source/AssetPlacementEdMode/Private/SAssetPlacementPalette.h
+++ b/Engine/Plugins/Experimental/AssetPlacementEdMode/Source/AssetPlacementEdMode/Private/SAssetPlacementPalette.h
@@ -79,6 +79,11 @@ private:
 	void ClearPalette();
 	void OnClearPalette();
 	void SetPaletteToAssetDataList(TArrayView<const FAssetData> InAssetDatas);
+	
+	/**
+	 * Removes the currently selected items from the palette.
+	 */
+	void OnRemoveSelectedItemsFromPalette();
 
 	/** Refreshes the active palette view widget */
 	void RefreshActivePaletteViewWidget();
