Revert code that was work arounding the fact that the SScrollBox slots were not configurable enough.

tests:
-placement editor drag n drop
-state tree view while resizing the scrollbox.

diff --git a/Engine/Plugins/Experimental/AssetPlacementEdMode/Source/AssetPlacementEdMode/Private/AssetPlacementEdModeToolkit.cpp b/Engine/Plugins/Experimental/AssetPlacementEdMode/Source/AssetPlacementEdMode/Private/AssetPlacementEdModeToolkit.cpp
--- a/Engine/Plugins/Experimental/AssetPlacementEdMode/Source/AssetPlacementEdMode/Private/AssetPlacementEdModeToolkit.cpp
+++ b/Engine/Plugins/Experimental/AssetPlacementEdMode/Source/AssetPlacementEdMode/Private/AssetPlacementEdModeToolkit.cpp
@@ -91,8 +91,13 @@ TSharedRef<SDockTab> FAssetPlacementEdModeToolkit::CreateAssetPalette(const FSpa
 {
 	return SNew(SDockTab)
 		[
-			SNew(SAssetPlacementPalette)
-			.ItemDetailsView(PaletteItemDetailsView)
+			SNew(SScrollBox)
+			+ SScrollBox::Slot()
+			.FillSize(1.0f)
+			[
+				SNew(SAssetPlacementPalette)
+				.ItemDetailsView(PaletteItemDetailsView)
+			]
 		];
 }
 
diff --git a/Engine/Plugins/Experimental/AssetPlacementEdMode/Source/AssetPlacementEdMode/Private/SAssetPlacementPalette.cpp b/Engine/Plugins/Experimental/AssetPlacementEdMode/Source/AssetPlacementEdMode/Private/SAssetPlacementPalette.cpp
--- a/Engine/Plugins/Experimental/AssetPlacementEdMode/Source/AssetPlacementEdMode/Private/SAssetPlacementPalette.cpp
+++ b/Engine/Plugins/Experimental/AssetPlacementEdMode/Source/AssetPlacementEdMode/Private/SAssetPlacementPalette.cpp
@@ -78,110 +78,108 @@ void SAssetPlacementPalette::Construct(const FArguments& InArgs)
 	}
 	SetupContentBrowserMirroring(!bIsMirroringContentBrowser);
 
-	TSharedPtr<SWidget> TopBar = 
-		SNew(SBorder)
-		.BorderImage(FAppStyle::Get().GetBrush("DetailsView.CategoryTop"))
-		.BorderBackgroundColor(FLinearColor(.6f, .6f, .6f, 1.0f))
+	ChildSlot
+	[
+		SNew(SVerticalBox)
+		+ SVerticalBox::Slot()
+		.HAlign(HAlign_Fill)
+		.AutoHeight()
 		[
-			SNew(SVerticalBox)
-			+ SVerticalBox::Slot()
-			.AutoHeight()
-			.HAlign(HAlign_Fill)
+			// Top bar
+			SNew(SBorder)
+			.BorderImage(FAppStyle::Get().GetBrush("DetailsView.CategoryTop"))
+			.BorderBackgroundColor(FLinearColor(.6f, .6f, .6f, 1.0f))
 			[
-				SNew(SCheckBox)
-				.Type(ESlateCheckBoxType::Type::ToggleButton)
-				.IsChecked(bIsMirroringContentBrowser ? ECheckBoxState::Checked : ECheckBoxState::Unchecked)
-				.Style(&FAssetPlacementEdModeStyle::Get().GetWidgetStyle<FCheckBoxStyle>("ToggleButtonCheckBox"))
-				.Visibility(this, &SAssetPlacementPalette::GetContentBrowserMirrorVisibility)
-				.OnCheckStateChanged(this, &SAssetPlacementPalette::OnContentBrowserMirrorButtonClicked)
+				SNew(SVerticalBox)
+				+ SVerticalBox::Slot()
+				.AutoHeight()
+				.HAlign(HAlign_Fill)
 				[
-					SNew(STextBlock)
-					.Text(LOCTEXT("Placement_ToggleContentBrowserMirroring", "Mirror Content Browser Selection"))
-					.Justification(ETextJustify::Type::Center)
-					.TextStyle(&FAssetPlacementEdModeStyle::Get().GetWidgetStyle<FTextBlockStyle>("ButtonText"))
-					.ToolTipText(LOCTEXT("Placement_ToggleContentBrowserMirroring_ToolTip", "Toggles palette to mirror the active content browser selection."))
+					SNew(SCheckBox)
+					.Type(ESlateCheckBoxType::Type::ToggleButton)
+					.IsChecked(bIsMirroringContentBrowser ? ECheckBoxState::Checked : ECheckBoxState::Unchecked)
+					.Style(&FAssetPlacementEdModeStyle::Get().GetWidgetStyle<FCheckBoxStyle>("ToggleButtonCheckBox"))
+					.Visibility(this, &SAssetPlacementPalette::GetContentBrowserMirrorVisibility)
+					.OnCheckStateChanged(this, &SAssetPlacementPalette::OnContentBrowserMirrorButtonClicked)
+					[
+						SNew(STextBlock)
+						.Text(LOCTEXT("Placement_ToggleContentBrowserMirroring", "Mirror Content Browser Selection"))
+						.Justification(ETextJustify::Type::Center)
+						.TextStyle(&FAssetPlacementEdModeStyle::Get().GetWidgetStyle<FTextBlockStyle>("ButtonText"))
+						.ToolTipText(LOCTEXT("Placement_ToggleContentBrowserMirroring_ToolTip", "Toggles palette to mirror the active content browser selection."))
+					]
 				]
-			]
-			+SVerticalBox::Slot()
-			.AutoHeight()
-			[
-				SNew(SObjectPropertyEntryBox)
-				.AllowedClass(UPlacementPaletteAsset::StaticClass())
-				.ObjectPath(this, &SAssetPlacementPalette::GetPalettePath)
-				.OnObjectChanged(this, &SAssetPlacementPalette::OnSetPaletteAsset)
-				.ThumbnailPool(ThumbnailPool)
-				.Visibility(this, &SAssetPlacementPalette::GetPaletteAssetPropertyBoxVisible)
-				.CustomContentSlot()
+				+SVerticalBox::Slot()
+				.AutoHeight()
 				[
-					SNew(SBox)
-					.HAlign(HAlign_Left)
-					.VAlign(VAlign_Center)
+					SNew(SObjectPropertyEntryBox)
+					.AllowedClass(UPlacementPaletteAsset::StaticClass())
+					.ObjectPath(this, &SAssetPlacementPalette::GetPalettePath)
+					.OnObjectChanged(this, &SAssetPlacementPalette::OnSetPaletteAsset)
+					.ThumbnailPool(ThumbnailPool)
+					.Visibility(this, &SAssetPlacementPalette::GetPaletteAssetPropertyBoxVisible)
+					.CustomContentSlot()
 					[
-						SNew(SHorizontalBox)
-						+ SHorizontalBox::Slot()
-						.VAlign(VAlign_Center)
-						.AutoWidth()
-						[
-							PropertyCustomizationHelpers::MakeClearButton(FSimpleDelegate::CreateSP(this, &SAssetPlacementPalette::OnResetPaletteAssetClicked), LOCTEXT("ResetPaletteAssetTooltip", "Clear the current palette asset, and use the user's local palette."), TAttribute<bool>::CreateLambda([this]() { return PalettePath.IsValid(); }))
-						]
-						+ SHorizontalBox::Slot()
+						SNew(SBox)
+						.HAlign(HAlign_Left)
 						.VAlign(VAlign_Center)
-						.AutoWidth()
-						.Padding(2.0f, 0.0f)
 						[
-							PropertyCustomizationHelpers::MakeSaveButton(FSimpleDelegate::CreateSP(this, &SAssetPlacementPalette::OnSavePaletteAssetClicked), LOCTEXT("SaveAssetPaletteTooltip", "Save changes to the current palette asset"), TAttribute<bool>::CreateLambda([this]() { return PalettePath.IsValid() || (PaletteItems.Num() > 0); }))
+							SNew(SHorizontalBox)
+							+ SHorizontalBox::Slot()
+							.VAlign(VAlign_Center)
+							.AutoWidth()
+							[
+								PropertyCustomizationHelpers::MakeClearButton(FSimpleDelegate::CreateSP(this, &SAssetPlacementPalette::OnResetPaletteAssetClicked), LOCTEXT("ResetPaletteAssetTooltip", "Clear the current palette asset, and use the user's local palette."), TAttribute<bool>::CreateLambda([this]() { return PalettePath.IsValid(); }))
+							]
+							+ SHorizontalBox::Slot()
+							.VAlign(VAlign_Center)
+							.AutoWidth()
+							.Padding(2.0f, 0.0f)
+							[
+								PropertyCustomizationHelpers::MakeSaveButton(FSimpleDelegate::CreateSP(this, &SAssetPlacementPalette::OnSavePaletteAssetClicked), LOCTEXT("SaveAssetPaletteTooltip", "Save changes to the current palette asset"), TAttribute<bool>::CreateLambda([this]() { return PalettePath.IsValid() || (PaletteItems.Num() > 0); }))
+							]
 						]
 					]
 				]
-			]
 
-			+ SVerticalBox::Slot()
-			.VAlign(VAlign_Center)
-			.AutoHeight()
-			[
-				SNew(SHorizontalBox)
-				+ SHorizontalBox::Slot()
-				.FillWidth(.75f)
+				+ SVerticalBox::Slot()
+				.VAlign(VAlign_Center)
+				.AutoHeight()
 				[
-					SAssignNew(SearchBoxPtr, SSearchBox)
-					.HintText(LOCTEXT("SearchPlacementPaletteHint", "Search Palette"))
-					.OnTextChanged(this, &SAssetPlacementPalette::OnSearchTextChanged)
-				]
+					SNew(SHorizontalBox)
+					+ SHorizontalBox::Slot()
+					.FillWidth(.75f)
+					[
+						SAssignNew(SearchBoxPtr, SSearchBox)
+						.HintText(LOCTEXT("SearchPlacementPaletteHint", "Search Palette"))
+						.OnTextChanged(this, &SAssetPlacementPalette::OnSearchTextChanged)
+					]
 
-				// View Options
-				+ SHorizontalBox::Slot()
-				.AutoWidth()
-				[
-					SNew(SComboButton)
-					.ForegroundColor(FSlateColor::UseForeground())
-					.ButtonStyle(FAppStyle::Get(), "ToggleButton")
-					.OnGetMenuContent(this, &SAssetPlacementPalette::GetViewOptionsMenuContent)
-					.ButtonContent()
+					// View Options
+					+ SHorizontalBox::Slot()
+					.AutoWidth()
 					[
-						SNew(SBox)
-						.VAlign(VAlign_Center)
+						SNew(SComboButton)
+						.ForegroundColor(FSlateColor::UseForeground())
+						.ButtonStyle(FAppStyle::Get(), "ToggleButton")
+						.OnGetMenuContent(this, &SAssetPlacementPalette::GetViewOptionsMenuContent)
+						.ButtonContent()
 						[
-							SNew(SImage)
-							.Image(FAppStyle::Get().GetBrush("GenericViewButton"))
+							SNew(SBox)
+							.VAlign(VAlign_Center)
+							[
+								SNew(SImage)
+								.Image(FAppStyle::Get().GetBrush("GenericViewButton"))
+							]
 						]
 					]
 				]
 			]
-		];
-
-	ChildSlot
-	[
-		SNew(SVerticalBox)
-		+ SVerticalBox::Slot()
-		.HAlign(HAlign_Fill)
-		.AutoHeight()
-		[
-			TopBar.ToSharedRef()
 		]
 
 		+ SVerticalBox::Slot()
-		.FillHeight(1.0)
-		[
+			.HAlign(HAlign_Fill)
+			.VAlign(VAlign_Fill)[
 			SNew(SAssetDropTarget)
 			.OnAreAssetsAcceptableForDrop(this, &SAssetPlacementPalette::OnAreAssetsValidForDrop)
 			.OnAssetsDropped(this, &SAssetPlacementPalette::HandlePlacementDropped)
@@ -190,24 +188,7 @@ void SAssetPlacementPalette::Construct(const FArguments& InArgs)
 				SNew(SOverlay)
 				+ SOverlay::Slot()
 				[
-					SNew(SBox)
-					.HeightOverride_Lambda([TopBarWeakPtr = TWeakPtr<SWidget>(TopBar)]()
-					{
-						if (TSharedPtr<SWidget> ParentWidgetPin = TopBarWeakPtr.Pin())
-						{
-							return FOptionalSize(ParentWidgetPin->GetTickSpaceGeometry().GetLocalSize().Y - 1);
-						}
-						return FOptionalSize();
-					})
-				]
-
-				+ SOverlay::Slot()
-				[
-					SNew(SScrollBox)
-					+ SScrollBox::Slot()
-					[
 						CreatePaletteViews()
-					]
 				]
 
 				+ SOverlay::Slot()
diff --git a/Engine/Plugins/Runtime/StateTree/Source/StateTreeEditorModule/Private/SStateTreeView.cpp b/Engine/Plugins/Runtime/StateTree/Source/StateTreeEditorModule/Private/SStateTreeView.cpp
--- a/Engine/Plugins/Runtime/StateTree/Source/StateTreeEditorModule/Private/SStateTreeView.cpp
+++ b/Engine/Plugins/Runtime/StateTree/Source/StateTreeEditorModule/Private/SStateTreeView.cpp
@@ -124,6 +124,7 @@ void SStateTreeView::Construct(const FArguments& InArgs, TSharedRef<FStateTreeVi
 				.Orientation(Orient_Horizontal)
 				.ExternalScrollbar(HorizontalScrollBar)
 				+SScrollBox::Slot()
+				.FillSize(1.0f)
 				[
 					TreeView.ToSharedRef()
 				]
diff --git a/Engine/Plugins/Runtime/StateTree/Source/StateTreeEditorModule/Private/SStateTreeViewRow.cpp b/Engine/Plugins/Runtime/StateTree/Source/StateTreeEditorModule/Private/SStateTreeViewRow.cpp
--- a/Engine/Plugins/Runtime/StateTree/Source/StateTreeEditorModule/Private/SStateTreeViewRow.cpp
+++ b/Engine/Plugins/Runtime/StateTree/Source/StateTreeEditorModule/Private/SStateTreeViewRow.cpp
@@ -35,85 +35,79 @@ void SStateTreeViewRow::Construct(const FArguments& InArgs, const TSharedRef<STa
 	WeakState = InState;
 
 	STableRow<TWeakObjectPtr<UStateTreeState>>::ConstructInternal(STableRow::FArguments()
-        .Padding(5.0f)
-        .OnDragDetected(this, &SStateTreeViewRow::HandleDragDetected)
-        .OnCanAcceptDrop(this, &SStateTreeViewRow::HandleCanAcceptDrop)
-        .OnAcceptDrop(this, &SStateTreeViewRow::HandleAcceptDrop)
-        .Style(&FStateTreeEditorStyle::Get()->GetWidgetStyle<FTableRowStyle>("StateTree.Selection"))
-    , InOwnerTableView);
+		.Padding(5.0f)
+		.OnDragDetected(this, &SStateTreeViewRow::HandleDragDetected)
+		.OnCanAcceptDrop(this, &SStateTreeViewRow::HandleCanAcceptDrop)
+		.OnAcceptDrop(this, &SStateTreeViewRow::HandleAcceptDrop)
+		.Style(&FStateTreeEditorStyle::Get()->GetWidgetStyle<FTableRowStyle>("StateTree.Selection"))
+	, InOwnerTableView);
 
 	static const FLinearColor TasksBackground = FLinearColor(FColor(17, 117, 131));
 	static const FLinearColor LinkBackground = FLinearColor(FColor(84, 84, 84));
 
 	this->ChildSlot
-    .HAlign(HAlign_Fill)
-    [
-	    SNew(SBox)
-	    .MinDesiredWidth_Lambda([ViewBox]()
-			{
-    			// Make the row at least as wide as the view.
-	    		// The -1 is needed or we'll see a scrollbar. 
-				return ViewBox->GetTickSpaceGeometry().GetLocalSize().X - 1;
-			})
+	.HAlign(HAlign_Fill)
+	[
+		SNew(SBox)
 		[
-	        SNew(SHorizontalBox)
-	        + SHorizontalBox::Slot()
-	        .VAlign(VAlign_Fill)
-	        .HAlign(HAlign_Left)
-	        .AutoWidth()
-	        [
-	            SNew(SExpanderArrow, SharedThis(this))
-	            .ShouldDrawWires(true)
-	            .IndentAmount(32)
-	            .BaseIndentLevel(0)
-	        ]
-
-	        // State Box
-	        + SHorizontalBox::Slot()
-	        .VAlign(VAlign_Center)
-	        .Padding(FMargin(0.0f, 4.0f))
-	        .AutoWidth()
-	        [
-	            SNew(SBox)
-	            .HeightOverride(28.0f)
-	            .VAlign(VAlign_Fill)
-	            [
-	                SNew(SBorder)
-	                .BorderImage(FAppStyle::GetBrush("WhiteBrush"))
-	                .BorderBackgroundColor(this, &SStateTreeViewRow::GetTitleColor)
-	                .Padding(FMargin(16.0f, 0.0f, 16.0f, 0.0f))
-	                [
-	                    // Conditions icon
-	                    SNew(SHorizontalBox)
-	                    + SHorizontalBox::Slot()
-	                    .VAlign(VAlign_Center)
-	                    .AutoWidth()
-	                    [
-	                        SNew(SBox)
-	                        .Padding(FMargin(0.0f, 0.0f, 4.0f, 0.0f))
-	                        .Visibility(this, &SStateTreeViewRow::GetConditionVisibility)
-	                        [
-	                            SNew(SImage)
+			SNew(SHorizontalBox)
+			+ SHorizontalBox::Slot()
+			.VAlign(VAlign_Fill)
+			.HAlign(HAlign_Left)
+			.AutoWidth()
+			[
+				SNew(SExpanderArrow, SharedThis(this))
+				.ShouldDrawWires(true)
+				.IndentAmount(32)
+				.BaseIndentLevel(0)
+			]
+
+			// State Box
+			+ SHorizontalBox::Slot()
+			.VAlign(VAlign_Center)
+			.Padding(FMargin(0.0f, 4.0f))
+			.AutoWidth()
+			[
+				SNew(SBox)
+				.HeightOverride(28.0f)
+				.VAlign(VAlign_Fill)
+				[
+					SNew(SBorder)
+					.BorderImage(FAppStyle::GetBrush("WhiteBrush"))
+					.BorderBackgroundColor(this, &SStateTreeViewRow::GetTitleColor)
+					.Padding(FMargin(16.0f, 0.0f, 16.0f, 0.0f))
+					[
+						// Conditions icon
+						SNew(SHorizontalBox)
+						+ SHorizontalBox::Slot()
+						.VAlign(VAlign_Center)
+						.AutoWidth()
+						[
+							SNew(SBox)
+							.Padding(FMargin(0.0f, 0.0f, 4.0f, 0.0f))
+							.Visibility(this, &SStateTreeViewRow::GetConditionVisibility)
+							[
+								SNew(SImage)
 								.Image(FAppStyle::Get().GetBrush("Icons.Help"))
 							]
 						]
 
 						// Selector icon
-	                    + SHorizontalBox::Slot()
-	                    .VAlign(VAlign_Center)
-	                    .AutoWidth()
-	                    [
-	                        SNew(SBox)
-	                        .Padding(FMargin(0.0f, 0.0f, 4.0f, 0.0f))
-	                        .Visibility(this, &SStateTreeViewRow::GetSelectorVisibility)
-	                        [
-	                            SNew(STextBlock)
-	                            .Text(this, &SStateTreeViewRow::GetSelectorDesc)
-	                            .TextStyle(FStateTreeEditorStyle::Get(), "StateTree.Icon")
-	                        ]
-	                    ]
-
-	                    // State Name
+						+ SHorizontalBox::Slot()
+						.VAlign(VAlign_Center)
+						.AutoWidth()
+						[
+							SNew(SBox)
+							.Padding(FMargin(0.0f, 0.0f, 4.0f, 0.0f))
+							.Visibility(this, &SStateTreeViewRow::GetSelectorVisibility)
+							[
+								SNew(STextBlock)
+								.Text(this, &SStateTreeViewRow::GetSelectorDesc)
+								.TextStyle(FStateTreeEditorStyle::Get(), "StateTree.Icon")
+							]
+						]
+
+						// State Name
 						+ SHorizontalBox::Slot()
 						.VAlign(VAlign_Center)
 						.AutoWidth()
@@ -256,7 +250,7 @@ void SStateTreeViewRow::Construct(const FArguments& InArgs, const TSharedRef<STa
 				[
 					SNew(STextBlock)
 					.Text(FEditorFontGlyphs::Check_Circle)
-	                .ColorAndOpacity(FLinearColor(FColor(110,143,67)))
+					.ColorAndOpacity(FLinearColor(FColor(110,143,67)))
 					.TextStyle(FStateTreeEditorStyle::Get(), "StateTree.Icon")
 					.Visibility(this, &SStateTreeViewRow::GetSucceededTransitionVisibility)
 				]
@@ -295,7 +289,7 @@ void SStateTreeViewRow::Construct(const FArguments& InArgs, const TSharedRef<STa
 				[
 					SNew(STextBlock)
 					.Text(FEditorFontGlyphs::Times_Circle)
-	                .ColorAndOpacity(FLinearColor(FColor(187,77,42)))
+					.ColorAndOpacity(FLinearColor(FColor(187,77,42)))
 					.TextStyle(FStateTreeEditorStyle::Get(), "StateTree.Icon")
 					.Visibility(this, &SStateTreeViewRow::GetFailedTransitionVisibility)
 				]
@@ -333,20 +327,20 @@ void SStateTreeViewRow::Construct(const FArguments& InArgs, const TSharedRef<STa
 				[
 					SNew(SImage)
 					.Image(FAppStyle::Get().GetBrush("Icons.Help"))
-	                .ColorAndOpacity(FLinearColor(FColor(31,151,167)))
+					.ColorAndOpacity(FLinearColor(FColor(31,151,167)))
 					.Visibility(this, &SStateTreeViewRow::GetConditionalTransitionsVisibility)
 				]
-	            + SHorizontalBox::Slot()
-	            .VAlign(VAlign_Center)
-	            .Padding(FMargin(4.0f, 0.0f, 0, 0))
-	            .AutoWidth()
-	            [
-	                SNew(STextBlock)
-	                .Text(FEditorFontGlyphs::Long_Arrow_Right)
-	                .TextStyle(FStateTreeEditorStyle::Get(), "StateTree.Icon")
-	                .Visibility(this, &SStateTreeViewRow::GetConditionalTransitionsVisibility)
-	            ]
-	            + SHorizontalBox::Slot()
+				+ SHorizontalBox::Slot()
+				.VAlign(VAlign_Center)
+				.Padding(FMargin(4.0f, 0.0f, 0, 0))
+				.AutoWidth()
+				[
+					SNew(STextBlock)
+					.Text(FEditorFontGlyphs::Long_Arrow_Right)
+					.TextStyle(FStateTreeEditorStyle::Get(), "StateTree.Icon")
+					.Visibility(this, &SStateTreeViewRow::GetConditionalTransitionsVisibility)
+				]
+				+ SHorizontalBox::Slot()
 				.VAlign(VAlign_Center)
 				.AutoWidth()
 				.Padding(FMargin(4.0f, 0, 0, 0))
